ControlRigSpline: Add parallel transport node

diff --git a/Engine/Plugins/Experimental/ControlRigSpline/Source/ControlRigSpline/Private/ControlRigSplineTypes.cpp b/Engine/Plugins/Experimental/ControlRigSpline/Source/ControlRigSpline/Private/ControlRigSplineTypes.cpp
--- a/Engine/Plugins/Experimental/ControlRigSpline/Source/ControlRigSpline/Private/ControlRigSplineTypes.cpp
+++ b/Engine/Plugins/Experimental/ControlRigSpline/Source/ControlRigSpline/Private/ControlRigSplineTypes.cpp
@@ -4,6 +4,11 @@ void FControlRigSpline::SetControlPoints(const TArray<FVector>& InPoints, const
 {
 	const int32 ControlPointsCount = InPoints.Num();
 
+	if (ControlPointsCount < 4)
+	{
+		return;
+	}
+
 	switch (SplineMode)
 	{
 		case ESplineType::BSpline:
@@ -22,6 +27,7 @@ void FControlRigSpline::SetControlPoints(const TArray<FVector>& InPoints, const
 			{
 				ts_bspline_set_control_point_at(BSpline->data(), i, &InPoints[i].X, nullptr);
 			}
+			
 			break;
 		}
 		case ESplineType::Hermite:
@@ -150,3 +156,57 @@ FVector FControlRigSpline::PositionAtParam(const float InParam) const
 
 	return Position;
 }
+
+FVector FControlRigSpline::TangentAtParam(const float InParam) const
+{
+	if (!BSpline.IsValid())
+	{
+		return FVector();
+	}
+
+	const float ClampedU = FMath::Clamp<float>(InParam, 0.f, 1.f);
+	const float Epsilon = 0.001f;
+
+	FVector Tangent;
+	std::vector<tinyspline::real> Result;
+	switch (SplineMode)
+	{
+		case ESplineType::BSpline:
+		{
+			float EpsilonMinus = ClampedU - Epsilon;
+			float EpsilonPlus = ClampedU + Epsilon;
+			if (EpsilonMinus < 0)
+			{
+				EpsilonMinus = 0.f;
+				EpsilonPlus = 2*Epsilon;
+			}
+			if (EpsilonPlus > 1.f)
+			{
+				EpsilonMinus = 1.f - 2*Epsilon;
+				EpsilonPlus = 1.f;
+			}
+			Result = BSpline->eval(EpsilonMinus).result();
+			FVector PositionMinus(Result[0], Result[1], Result[2]);
+			Result = BSpline->eval(EpsilonPlus).result();
+			FVector PositionPlus(Result[0], Result[1], Result[2]);
+			Tangent = PositionPlus - PositionMinus;
+			break;
+		}
+		case ESplineType::Hermite:
+		{
+			Result = BSpline->bisect(ClampedU - Epsilon).result();
+			FVector PositionMinus(Result[1], Result[2], Result[3]);
+			Result = BSpline->bisect(ClampedU + Epsilon).result();
+			FVector PositionPlus(Result[1], Result[2], Result[3]);
+			Tangent = PositionPlus - PositionMinus;
+			break;
+		}
+		default:
+		{
+			checkNoEntry(); // Unknown Spline Mode
+			break;
+		}
+	}
+
+	return Tangent;
+}
diff --git a/Engine/Plugins/Experimental/ControlRigSpline/Source/ControlRigSpline/Private/ControlRigSplineUnits.cpp b/Engine/Plugins/Experimental/ControlRigSpline/Source/ControlRigSpline/Private/ControlRigSplineUnits.cpp
--- a/Engine/Plugins/Experimental/ControlRigSpline/Source/ControlRigSpline/Private/ControlRigSplineUnits.cpp
+++ b/Engine/Plugins/Experimental/ControlRigSpline/Source/ControlRigSpline/Private/ControlRigSplineUnits.cpp
@@ -47,6 +47,43 @@ FRigUnit_PositionFromControlRigSpline_Execute()
 	}
 }
 
+FRigUnit_TransformFromControlRigSpline_Execute()
+{
+	switch (Context.State)
+	{
+		case EControlRigState::Init:
+		case EControlRigState::Update:
+		{
+			FVector UpVectorNormalized = UpVector;
+			UpVectorNormalized.Normalize();
+
+			const float ClampedU = FMath::Clamp<float>(U, 0.f, 1.f);
+			const float ClampedTwist = FMath::Clamp<float>(Twist, -180.f, 180.f);
+		
+			FVector Tangent = Spline.TangentAtParam(ClampedU);
+
+			// Check if Tangent can be normalized. If not, keep the same tangent as before.
+			if (!Tangent.Normalize())
+			{
+				Tangent = Transform.ToMatrixNoScale().GetUnitAxis(EAxis::X);
+			}
+			FVector Binormal = FVector::CrossProduct(Tangent, UpVectorNormalized);
+			Binormal = Binormal.RotateAngleAxis(ClampedTwist * ClampedU, Tangent);
+
+			FMatrix RotationMatrix = FRotationMatrix::MakeFromXZ(Tangent, Binormal);
+
+			Transform.SetFromMatrix(RotationMatrix);
+			Transform.SetTranslation(Spline.PositionAtParam(U));
+			break;
+		}
+		default:
+		{
+			checkNoEntry(); // Execute is only defined for Init and Update
+			break;
+		}
+	}
+}
+
 FRigUnit_DrawControlRigSpline_Execute()
 {
 	if (Context.State == EControlRigState::Init)
diff --git a/Engine/Plugins/Experimental/ControlRigSpline/Source/ControlRigSpline/Public/ControlRigSplineTypes.h b/Engine/Plugins/Experimental/ControlRigSpline/Source/ControlRigSpline/Public/ControlRigSplineTypes.h
--- a/Engine/Plugins/Experimental/ControlRigSpline/Source/ControlRigSpline/Public/ControlRigSplineTypes.h
+++ b/Engine/Plugins/Experimental/ControlRigSpline/Source/ControlRigSpline/Public/ControlRigSplineTypes.h
@@ -63,6 +63,15 @@ struct CONTROLRIGSPLINE_API FControlRigSpline
 	*/
 	FVector PositionAtParam(const float InParam) const;
 
+	/**
+	* Given an InParam float in [0, 1], will return the tangent vector of the spline at that point. 
+	* Note that this vector is not normalized.
+	*
+	* @param InParam	The parameter between [0, 1] to query.
+	* @return			The tangent of the spline at InParam.
+	*/
+	FVector TangentAtParam(const float InParam) const;
+
 	// Auxiliary control points array, which is different from the Points array
 	// when SplineMode is Hermite
 	TArray<float> ControlPointsArray;
diff --git a/Engine/Plugins/Experimental/ControlRigSpline/Source/ControlRigSpline/Public/ControlRigSplineUnits.h b/Engine/Plugins/Experimental/ControlRigSpline/Source/ControlRigSpline/Public/ControlRigSplineUnits.h
--- a/Engine/Plugins/Experimental/ControlRigSpline/Source/ControlRigSpline/Public/ControlRigSplineUnits.h
+++ b/Engine/Plugins/Experimental/ControlRigSpline/Source/ControlRigSpline/Public/ControlRigSplineUnits.h
@@ -18,6 +18,11 @@ struct CONTROLRIGSPLINE_API FRigUnit_ControlRigSplineFromPoints : public FRigUni
 {
 	GENERATED_BODY()
 
+	FRigUnit_ControlRigSplineFromPoints()
+	{
+		SplineMode = ESplineType::BSpline;
+	}
+
 	/** Execute logic for this rig unit */
 	RIGVM_METHOD()
 	virtual void Execute(const FRigUnitContext& Context) override;
@@ -37,6 +42,12 @@ struct CONTROLRIGSPLINE_API FRigUnit_PositionFromControlRigSpline : public FRigU
 {
 	GENERATED_BODY()
 
+	FRigUnit_PositionFromControlRigSpline()
+	{
+		U = 0.f;
+		Position = FVector::ZeroVector;
+	}
+
 	/** Execute logic for this rig unit */
 	RIGVM_METHOD()
 	virtual void Execute(const FRigUnitContext& Context) override;
@@ -51,6 +62,39 @@ struct CONTROLRIGSPLINE_API FRigUnit_PositionFromControlRigSpline : public FRigU
 	FVector Position;
 };
 
+USTRUCT(meta = (DisplayName = "Transform From Control Rig Spline", Category = "Control Rig"))
+struct CONTROLRIGSPLINE_API FRigUnit_TransformFromControlRigSpline : public FRigUnit_ControlRigSplineBase
+{
+	GENERATED_BODY()
+
+	FRigUnit_TransformFromControlRigSpline()
+	{
+		UpVector = FVector::UpVector;
+		Twist = 0.f;
+		U = 0.f;
+		Transform = FTransform::Identity;
+	}
+
+	/** Execute logic for this rig unit */
+	RIGVM_METHOD()
+	virtual void Execute(const FRigUnitContext& Context) override;
+
+	UPROPERTY(meta = (Input))
+	FControlRigSpline Spline;
+
+	UPROPERTY(meta = (Input))
+	FVector UpVector;
+
+	UPROPERTY(meta = (Input))
+	float Twist;
+
+	UPROPERTY(meta = (Input))
+	float U;
+
+	UPROPERTY(meta = (Output))
+	FTransform Transform;
+};
+
 USTRUCT(meta = (DisplayName = "Draw Control Rig Spline", Category = "Control Rig"))
 struct CONTROLRIGSPLINE_API FRigUnit_DrawControlRigSpline : public FRigUnitMutable
 {
@@ -58,7 +102,6 @@ struct CONTROLRIGSPLINE_API FRigUnit_DrawControlRigSpline : public FRigUnitMutab
 
 	FRigUnit_DrawControlRigSpline()
 	{
-		Spline = FControlRigSpline();
 		Color = FLinearColor::Red;
 		Thickness = 1.f;
 		Detail = 16.f;
