ControlRigSpline: Not loading in Play4

diff --git a/Engine/Plugins/Experimental/ControlRigSpline/ControlRigSpline.uplugin b/Engine/Plugins/Experimental/ControlRigSpline/ControlRigSpline.uplugin
--- a/Engine/Plugins/Experimental/ControlRigSpline/ControlRigSpline.uplugin
+++ b/Engine/Plugins/Experimental/ControlRigSpline/ControlRigSpline.uplugin
@@ -25,7 +25,8 @@
 		{
 			"Name": "ControlRigSpline",
 			"Type": "Runtime",
-			"LoadingPhase": "Default"
+			"LoadingPhase": "Default",
+			"WhitelistPlatforms": [ "Win64" ]
 		}
 	]
 }
\ No newline at end of file
