ControlRigSpline: GetLengthFromSpline

diff --git a/Engine/Plugins/Experimental/ControlRigSpline/Source/ControlRigSpline/Private/ControlRigSplineTypes.cpp b/Engine/Plugins/Experimental/ControlRigSpline/Source/ControlRigSpline/Private/ControlRigSplineTypes.cpp
--- a/Engine/Plugins/Experimental/ControlRigSpline/Source/ControlRigSpline/Private/ControlRigSplineTypes.cpp
+++ b/Engine/Plugins/Experimental/ControlRigSpline/Source/ControlRigSpline/Private/ControlRigSplineTypes.cpp
@@ -1,95 +1,172 @@
-// Copyright Epic Games, Inc. All Rights Reserved.
-
 #include "ControlRigSplineTypes.h"
 
-void FControlRigSpline::SetControlPoints(const TArray<FVector>& InPoints, const bool forceRebuild)
+void FControlRigSpline::SetControlPoints(const TArray<FVector>& InPoints, const ESplineType SplineMode, const int32 SamplesPerSegment)
 {
 	const int32 ControlPointsCount = InPoints.Num();
-
 	if (ControlPointsCount < 4)
 	{
 		return;
 	}
 
-	switch (SplineMode)
+	if (!SplineData.IsValid())
 	{
-		case ESplineType::BSpline:
+		SplineData = MakeShared<FControlRigSplineImpl>();
+	}
+
+	bool bControlPointsChanged = InPoints != SplineData->ControlPoints;
+	bool bSplineModeChanged = SplineMode != SplineData->SplineMode;
+	bool bSamplesCountChanged = SamplesPerSegment != SplineData->SamplesPerSegment;
+	if (!bSplineModeChanged && !bControlPointsChanged && !bSamplesCountChanged)
+	{
+		return;
+	}
+
+	SplineData->ControlPoints = InPoints;
+	SplineData->SplineMode = SplineMode;
+	SplineData->SamplesPerSegment = SamplesPerSegment;
+
+	// If we need to update the spline because the controls points have changed, or the spline mode has changed
+	if (bControlPointsChanged || bSplineModeChanged)
+	{	
+		switch (SplineMode)
 		{
-			if (forceRebuild || !BSpline.IsValid() || BSpline->numControlPoints() != ControlPointsCount)
+			case ESplineType::BSpline:
 			{
-				BSpline = MakeShareable(new tinyspline::BSpline(ControlPointsCount, 3));
-			}
+				if (bSplineModeChanged || SplineData->Spline.numControlPoints() != ControlPointsCount)
+				{
+					SplineData->Spline = tinyspline::BSpline(ControlPointsCount, 3);
+				}
+
+				// Update the positions of the control points
+
+				// There's no guarantee that FVector is a tightly packed array of three floats. 
+				// We have SIMD versions where we waste a dummy float to align it on a 16 byte boundary,
+				// so we need to iterate updating the points one by one.
+				for (int32 i = 0; i < ControlPointsCount; ++i)
+				{
+					ts_bspline_set_control_point_at(SplineData->Spline.data(), i, &InPoints[i].X, nullptr);
+				}
 
-			// Update the positions of the control points
+				break;
+			}
+			case ESplineType::Hermite:
+			{
+				if (bSplineModeChanged || SplineData->Spline.numControlPoints() != ControlPointsCount)
+				{
+					// Create a new BSpline from the control points
+					// The first dimension acts as the query parameter, the rest of the dimensions are the XYZ of the control points
+					TArray<float> ControlPointsArray;
+					ControlPointsArray.SetNumUninitialized(ControlPointsCount * 4, false);
+				
+					for (int32 i = 0; i < ControlPointsCount; ++i)
+					{
+						// The control point with weight 0 is the second point
+						// The control point with weight 1 is the second to last point
+						// This way the first and last point act as tangents
+						ControlPointsArray[i * 4 + 0] = (i - 1) / (float)(ControlPointsCount - 3);
+
+						ControlPointsArray[i * 4 + 1] = InPoints[i][0];
+						ControlPointsArray[i * 4 + 2] = InPoints[i][1];
+						ControlPointsArray[i * 4 + 3] = InPoints[i][2];
+					}
+
+					SplineData->Spline = tinyspline::BSpline(ControlPointsCount, 4);
+					ts_bspline_interpolate_cubic_natural(ControlPointsArray.GetData(), ControlPointsCount, 4, SplineData->Spline.data(), nullptr);
+				}
+				else
+				{
+					// We just need to update the positions of the control points
+					float* Ctrlp = nullptr;
+					ts_bspline_control_point_at(SplineData->Spline.data(), 0, &Ctrlp, nullptr);
+					for (int32 i = 0; i < ControlPointsCount; ++i)
+					{
+						Ctrlp[i * 4 + 1] = InPoints[i][0];
+						Ctrlp[i * 4 + 2] = InPoints[i][1];
+						Ctrlp[i * 4 + 3] = InPoints[i][2];
+					}
+				}
 
-			// There's no guarantee that FVector is a tightly packed array of three floats. 
-			// We have SIMD versions where we waste a dummy float to align it on a 16 byte boundary,
-			// so we need to iterate updating the points one by one.
-			for (int i = 0; i < ControlPointsCount; ++i)
+				break;
+			}
+			default:
 			{
-				ts_bspline_set_control_point_at(BSpline->data(), i, &InPoints[i].X, nullptr);
+				checkNoEntry(); // Unknown Spline Mode
+				break;
 			}
-			
-			break;
 		}
-		case ESplineType::Hermite:
+	}
+
+	// If curve has changed, or sample count has changed, recompute the cache
+	if (bControlPointsChanged || bSplineModeChanged || bSamplesCountChanged)
+	{
+		switch (SplineMode)
 		{
-			if (forceRebuild || !BSpline.IsValid() || BSpline->numControlPoints() != ControlPointsCount)
+			case ESplineType::BSpline:
 			{
-				// Create a new BSpline from the control points
-				// The first dimension acts as the query parameter, the rest of the dimensions are the XYZ of the control points
-				ControlPointsArray.SetNumUninitialized(ControlPointsCount * 4, false);
-				
-				for (int32 i = 0; i < ControlPointsCount; ++i)
+				// Cache sample positions of the spline
+				float* SamplesPtr = nullptr;
+				size_t ActualSamplesPerSegment = 0;
+				ts_bspline_sample(SplineData->Spline.data(), (ControlPointsCount - 1) * SamplesPerSegment, &SamplesPtr, &ActualSamplesPerSegment, nullptr);
+				SplineData->SamplesArray.SetNumUninitialized((ControlPointsCount - 1) * SamplesPerSegment, false);
+				for (int32 i = 0; i < ControlPointsCount - 1; ++i)
 				{
-					// The control point with weight 0 is the second point
-					// The control point with weight 1 is the second to last point
-					// This way the first and last point act as tangents
-					ControlPointsArray[i * 4 + 0] = (i - 1) / (float)(ControlPointsCount - 3);
-
-					ControlPointsArray[i * 4 + 1] = InPoints[i][0];
-					ControlPointsArray[i * 4 + 2] = InPoints[i][1];
-					ControlPointsArray[i * 4 + 3] = InPoints[i][2];
+					for (int32 j = 0; j < SamplesPerSegment; ++j)
+					{
+						memcpy(&SplineData->SamplesArray[i * SamplesPerSegment + j].X, SamplesPtr + (i * SamplesPerSegment + j) * 3, sizeof(float) * 3);
+					}
 				}
 
-				BSpline = MakeShareable(new tinyspline::BSpline(ControlPointsCount, 4));
-				ts_bspline_interpolate_cubic_natural(ControlPointsArray.GetData(), ControlPointsCount, 4, BSpline->data(), nullptr);
+				// tinySpline will allocate the samples array, but does not free that memory. We need to take care of that.
+				free(SamplesPtr);
+
+				break;
 			}
-			else
+			case ESplineType::Hermite:
 			{
-				// We just need to update the positions of the control points
-				float* Ctrlp = nullptr;
-				ts_bspline_control_point_at(BSpline->data(), 0, &Ctrlp, nullptr);
-				for (int32 i = 0; i < ControlPointsCount; ++i)
+				// Cache sample positions of the spline
+				float* SamplesPtr = nullptr;
+				size_t ActualSamplesPerSegment = 0;
+				ts_bspline_sample(SplineData->Spline.data(), (ControlPointsCount - 1) * SamplesPerSegment, &SamplesPtr, &ActualSamplesPerSegment, nullptr);
+				SplineData->SamplesArray.SetNumUninitialized((ControlPointsCount - 3) * SamplesPerSegment);
+				for (int32 i = 0; i < ControlPointsCount - 3; ++i)
 				{
-					Ctrlp[i * 4 + 1] = InPoints[i][0];
-					Ctrlp[i * 4 + 2] = InPoints[i][1];
-					Ctrlp[i * 4 + 3] = InPoints[i][2];
-
-					ControlPointsArray[i * 4 + 1] = InPoints[i][0];
-					ControlPointsArray[i * 4 + 2] = InPoints[i][1];
-					ControlPointsArray[i * 4 + 3] = InPoints[i][2];
+					for (int32 j = 0; j < SamplesPerSegment; ++j)
+					{
+						memcpy(&SplineData->SamplesArray[i * SamplesPerSegment + j].X, SamplesPtr + ((i + 1) * SamplesPerSegment + j) * 4 + 1, sizeof(float) * 3);
+					}
 				}
+
+				// tinySpline will allocate the samples array, but does not free that memory. We need to take care of that.
+				free(SamplesPtr);
+
+				break;
+			}
+			default:
+			{
+				checkNoEntry(); // Unknown Spline Mode
+				break;
 			}
-			break;
 		}
-		default:
+
+		// Cache accumulated length at sample array
+		SplineData->AccumulatedLenth.SetNumUninitialized(SplineData->SamplesArray.Num(), false);
+		SplineData->AccumulatedLenth[0] = 0.f;
+		for (int32 i = 1; i < SplineData->SamplesArray.Num(); ++i)
 		{
-			checkNoEntry(); // Unknown Spline Mode
-			break;
+			SplineData->AccumulatedLenth[i] = SplineData->AccumulatedLenth[i - 1] + FVector::Distance(SplineData->SamplesArray[i - 1], SplineData->SamplesArray[i]);
 		}
-	}
+	}	
 }
 
 int32 FControlRigSpline::GetControlPoints(TArray<FVector>& OutPoints) const
 {
-	const size_t Count = BSpline->numControlPoints();
+	const size_t Count = SplineData->Spline.numControlPoints();
 	OutPoints.Reset(Count);
 
 	float* Points;
-	ts_bspline_control_points(BSpline->data(), &Points, nullptr);
+	ts_bspline_control_points(SplineData->Spline.data(), &Points, nullptr);
 
-	switch (SplineMode)
+	switch (SplineData->SplineMode)
 	{
 		case ESplineType::BSpline:
 		{
@@ -116,99 +193,36 @@ int32 FControlRigSpline::GetControlPoints(TArray<FVector>& OutPoints) const
 		}
 	}
 
-
 	return Count;
 }
 
 FVector FControlRigSpline::PositionAtParam(const float InParam) const
 {
-	if (!BSpline.IsValid())
+	if (!SplineData.IsValid())
 	{
 		return FVector();
 	}
 
 	float ClampedU = FMath::Clamp<float>(InParam, 0.f, 1.f);
 
-	FVector Position;
-	std::vector<tinyspline::real> Result;
-	switch (SplineMode)
-	{
-		case ESplineType::BSpline:
-		{
-			Result = BSpline->eval(ClampedU).result();
-			Position[0] = Result[0];
-			Position[1] = Result[1];
-			Position[2] = Result[2];
-			break;
-		}
-		case ESplineType::Hermite:
-		{
-			Result = BSpline->bisect(ClampedU).result();
-			Position[0] = Result[1];
-			Position[1] = Result[2];
-			Position[2] = Result[3];
-			break;
-		}
-		default:
-		{
-			checkNoEntry(); // Unknown Spline Mode
-			break;
-		}
-	}
+	const int32 LastIndex = SplineData->SamplesArray.Num() - 1;
+	float fIndexPrev = ClampedU * LastIndex;	
+	int32 IndexPrev = FMath::Floor<int32>(fIndexPrev);
+	float ULocal = fIndexPrev - IndexPrev;
+	int32 IndexNext = (IndexPrev < LastIndex) ? IndexPrev + 1 : IndexPrev;
 
-	return Position;
+	return SplineData->SamplesArray[IndexPrev] * (1.f - ULocal) + SplineData->SamplesArray[IndexNext] * ULocal;
 }
 
 FVector FControlRigSpline::TangentAtParam(const float InParam) const
 {
-	if (!BSpline.IsValid())
+	if (!SplineData.IsValid())
 	{
 		return FVector();
 	}
 
 	const float ClampedU = FMath::Clamp<float>(InParam, 0.f, 1.f);
-	const float Epsilon = 0.001f;
-
-	FVector Tangent;
-	std::vector<tinyspline::real> Result;
-	switch (SplineMode)
-	{
-		case ESplineType::BSpline:
-		{
-			float EpsilonMinus = ClampedU - Epsilon;
-			float EpsilonPlus = ClampedU + Epsilon;
-			if (EpsilonMinus < 0)
-			{
-				EpsilonMinus = 0.f;
-				EpsilonPlus = 2*Epsilon;
-			}
-			if (EpsilonPlus > 1.f)
-			{
-				EpsilonMinus = 1.f - 2*Epsilon;
-				EpsilonPlus = 1.f;
-			}
-			Result = BSpline->eval(EpsilonMinus).result();
-			FVector PositionMinus(Result[0], Result[1], Result[2]);
-			Result = BSpline->eval(EpsilonPlus).result();
-			FVector PositionPlus(Result[0], Result[1], Result[2]);
-			Tangent = PositionPlus - PositionMinus;
-			break;
-		}
-		case ESplineType::Hermite:
-		{
-			Result = BSpline->bisect(ClampedU - Epsilon).result();
-			FVector PositionMinus(Result[1], Result[2], Result[3]);
-			Result = BSpline->bisect(ClampedU + Epsilon).result();
-			FVector PositionPlus(Result[1], Result[2], Result[3]);
-			Tangent = PositionPlus - PositionMinus;
-			break;
-		}
-		default:
-		{
-			checkNoEntry(); // Unknown Spline Mode
-			break;
-		}
-	}
 
-	return Tangent;
+	int32 IndexPrev = ClampedU * (SplineData->SamplesArray.Num()-2);
+	return SplineData->SamplesArray[IndexPrev+1] - SplineData->SamplesArray[IndexPrev];
 }
diff --git a/Engine/Plugins/Experimental/ControlRigSpline/Source/ControlRigSpline/Private/ControlRigSplineUnits.cpp b/Engine/Plugins/Experimental/ControlRigSpline/Source/ControlRigSpline/Private/ControlRigSplineUnits.cpp
--- a/Engine/Plugins/Experimental/ControlRigSpline/Source/ControlRigSpline/Private/ControlRigSplineUnits.cpp
+++ b/Engine/Plugins/Experimental/ControlRigSpline/Source/ControlRigSpline/Private/ControlRigSplineUnits.cpp
@@ -13,12 +13,7 @@ FRigUnit_ControlRigSplineFromPoints_Execute()
 		case EControlRigState::Init:
 		case EControlRigState::Update:
 		{
-			const int32 ControlPointsCount = Points.Num();
-		
-			bool bModeChanged = Spline.SplineMode != SplineMode;
-			Spline.SplineMode = SplineMode;
-
-			Spline.SetControlPoints(Points, bModeChanged);
+			Spline.SetControlPoints(Points, SplineMode, SamplesPerSegment);
 			break;
 		}
 		default:
@@ -58,7 +53,7 @@ FRigUnit_TransformFromControlRigSpline_Execute()
 			UpVectorNormalized.Normalize();
 
 			const float ClampedU = FMath::Clamp<float>(U, 0.f, 1.f);
-			const float ClampedTwist = FMath::Clamp<float>(Twist, -180.f, 180.f);
+			const float ClampedRoll = FMath::Clamp<float>(Roll, -180.f, 180.f);
 		
 			FVector Tangent = Spline.TangentAtParam(ClampedU);
 
@@ -68,8 +63,8 @@ FRigUnit_TransformFromControlRigSpline_Execute()
 				Tangent = Transform.ToMatrixNoScale().GetUnitAxis(EAxis::X);
 			}
 			FVector Binormal = FVector::CrossProduct(Tangent, UpVectorNormalized);
-			Binormal = Binormal.RotateAngleAxis(ClampedTwist * ClampedU, Tangent);
-
+			Binormal = Binormal.RotateAngleAxis(ClampedRoll * ClampedU, Tangent);
+			
 			FMatrix RotationMatrix = FRotationMatrix::MakeFromXZ(Tangent, Binormal);
 
 			Transform.SetFromMatrix(RotationMatrix);
@@ -111,3 +106,21 @@ FRigUnit_DrawControlRigSpline_Execute()
 
 	Context.DrawInterface->Instructions.Add(Instruction);
 }
+
+FRigUnit_GetLengthControlRigSpline_Execute()
+{
+	switch (Context.State)
+	{
+		case EControlRigState::Init:
+		case EControlRigState::Update:
+		{
+			Length = Spline.SplineData->AccumulatedLenth.Last();
+			break;
+		}
+		default:
+		{
+			checkNoEntry(); // Execute is only defined for Init and Update
+			break;
+		}
+	}
+}
\ No newline at end of file
diff --git a/Engine/Plugins/Experimental/ControlRigSpline/Source/ControlRigSpline/Public/ControlRigSplineTypes.h b/Engine/Plugins/Experimental/ControlRigSpline/Source/ControlRigSpline/Public/ControlRigSplineTypes.h
--- a/Engine/Plugins/Experimental/ControlRigSpline/Source/ControlRigSpline/Public/ControlRigSplineTypes.h
+++ b/Engine/Plugins/Experimental/ControlRigSpline/Source/ControlRigSpline/Public/ControlRigSplineTypes.h
@@ -24,18 +24,46 @@ enum class ESplineType : uint8
 	Max UMETA(Hidden),
 };
 
+USTRUCT()
+struct CONTROLRIGSPLINE_API FControlRigSplineImpl
+{
+	GENERATED_BODY()
+
+	FControlRigSplineImpl()	
+	{
+		SplineMode = ESplineType::BSpline;
+		SamplesPerSegment = 16;
+	}
+
+	// Spline type
+	ESplineType SplineMode;
+
+	// The control points to construct the spline
+	TArray<FVector> ControlPoints;
+
+	// The actual spline
+	tinyspline::BSpline Spline;
+
+	// Samples per segment, where segment is the portion between two control points
+	int32 SamplesPerSegment;
+
+	// Positions along the "real" curve (no samples in the first and last segments of a hermite spline)
+	TArray<FVector> SamplesArray;
+
+	// Accumulated length along the spline given by samples
+	TArray<float> AccumulatedLenth;
+};
+
 USTRUCT(BlueprintType)
 struct CONTROLRIGSPLINE_API FControlRigSpline 
 {
 	GENERATED_BODY()
 
 	FControlRigSpline()	{}
-	virtual ~FControlRigSpline() {}
 
-	// Spline type
-	ESplineType SplineMode;
+	virtual ~FControlRigSpline() {}
 
-	TSharedPtr<tinyspline::BSpline> BSpline;
+	TSharedPtr<FControlRigSplineImpl> SplineData;
 
 	/**
 	* Sets the control points in the spline. It will build the spline if needed, or forceRebuild is true,
@@ -45,7 +73,7 @@ struct CONTROLRIGSPLINE_API FControlRigSpline
 	* @param InPoints	The control points to set.
 	* @param forceRebuild	If true, will build the spline from scratch.
 	*/
-	void SetControlPoints(const TArray<FVector>& InPoints, const bool forceRebuild = false);
+	void SetControlPoints(const TArray<FVector>& InPoints, const ESplineType SplineMode = ESplineType::BSpline, const int32 SamplesPerSegment = 16);
 
 	/**
 	* Populates OutPoints and returns the number of control points in the spline.
@@ -71,8 +99,4 @@ struct CONTROLRIGSPLINE_API FControlRigSpline
 	* @return			The tangent of the spline at InParam.
 	*/
 	FVector TangentAtParam(const float InParam) const;
-
-	// Auxiliary control points array, which is different from the Points array
-	// when SplineMode is Hermite
-	TArray<float> ControlPointsArray;
 };
diff --git a/Engine/Plugins/Experimental/ControlRigSpline/Source/ControlRigSpline/Public/ControlRigSplineUnits.h b/Engine/Plugins/Experimental/ControlRigSpline/Source/ControlRigSpline/Public/ControlRigSplineUnits.h
--- a/Engine/Plugins/Experimental/ControlRigSpline/Source/ControlRigSpline/Public/ControlRigSplineUnits.h
+++ b/Engine/Plugins/Experimental/ControlRigSpline/Source/ControlRigSpline/Public/ControlRigSplineUnits.h
@@ -13,7 +13,7 @@ struct CONTROLRIGSPLINE_API FRigUnit_ControlRigSplineBase : public FRigUnit
 	GENERATED_BODY()
 };
 
-USTRUCT(meta = (DisplayName = "Control Rig Spline From Points", Category = "Control Rig"))
+USTRUCT(meta = (DisplayName = "Spline From Points", Keywords="Spline From Positions", Category = "Control Rig"))
 struct CONTROLRIGSPLINE_API FRigUnit_ControlRigSplineFromPoints : public FRigUnit_ControlRigSplineBase
 {
 	GENERATED_BODY()
@@ -21,6 +21,7 @@ struct CONTROLRIGSPLINE_API FRigUnit_ControlRigSplineFromPoints : public FRigUni
 	FRigUnit_ControlRigSplineFromPoints()
 	{
 		SplineMode = ESplineType::BSpline;
+		SamplesPerSegment = 16;
 	}
 
 	/** Execute logic for this rig unit */
@@ -30,14 +31,17 @@ struct CONTROLRIGSPLINE_API FRigUnit_ControlRigSplineFromPoints : public FRigUni
 	UPROPERTY(meta = (Input))
 	TArray<FVector> Points;
 
-	UPROPERTY(meta = (Output))
-	FControlRigSpline Spline;
-
 	UPROPERTY(meta = (Input))
 	ESplineType SplineMode;
+
+	UPROPERTY(meta = (Input))
+	int32 SamplesPerSegment;
+
+	UPROPERTY(meta = (Output))
+	FControlRigSpline Spline;
 };
 
-USTRUCT(meta = (DisplayName = "Position From Control Rig Spline", Category = "Control Rig"))
+USTRUCT(meta = (DisplayName = "Position From Spline", Keywords="Point From Spline", Category = "Control Rig"))
 struct CONTROLRIGSPLINE_API FRigUnit_PositionFromControlRigSpline : public FRigUnit_ControlRigSplineBase
 {
 	GENERATED_BODY()
@@ -62,7 +66,7 @@ struct CONTROLRIGSPLINE_API FRigUnit_PositionFromControlRigSpline : public FRigU
 	FVector Position;
 };
 
-USTRUCT(meta = (DisplayName = "Transform From Control Rig Spline", Category = "Control Rig"))
+USTRUCT(meta = (DisplayName = "Transform From Spline", Category = "Control Rig"))
 struct CONTROLRIGSPLINE_API FRigUnit_TransformFromControlRigSpline : public FRigUnit_ControlRigSplineBase
 {
 	GENERATED_BODY()
@@ -70,7 +74,7 @@ struct CONTROLRIGSPLINE_API FRigUnit_TransformFromControlRigSpline : public FRig
 	FRigUnit_TransformFromControlRigSpline()
 	{
 		UpVector = FVector::UpVector;
-		Twist = 0.f;
+		Roll = 0.f;
 		U = 0.f;
 		Transform = FTransform::Identity;
 	}
@@ -86,7 +90,7 @@ struct CONTROLRIGSPLINE_API FRigUnit_TransformFromControlRigSpline : public FRig
 	FVector UpVector;
 
 	UPROPERTY(meta = (Input))
-	float Twist;
+	float Roll;
 
 	UPROPERTY(meta = (Input))
 	float U;
@@ -95,7 +99,7 @@ struct CONTROLRIGSPLINE_API FRigUnit_TransformFromControlRigSpline : public FRig
 	FTransform Transform;
 };
 
-USTRUCT(meta = (DisplayName = "Draw Control Rig Spline", Category = "Control Rig"))
+USTRUCT(meta = (DisplayName = "Draw Spline", Category = "Control Rig"))
 struct CONTROLRIGSPLINE_API FRigUnit_DrawControlRigSpline : public FRigUnitMutable
 {
 	GENERATED_BODY()
@@ -124,3 +128,23 @@ struct CONTROLRIGSPLINE_API FRigUnit_DrawControlRigSpline : public FRigUnitMutab
 	int32 Detail;
 };
 
+USTRUCT(meta = (DisplayName = "Get Length Of Spline", Category = "Control Rig"))
+struct CONTROLRIGSPLINE_API FRigUnit_GetLengthControlRigSpline : public FRigUnit
+{
+	GENERATED_BODY()
+
+	FRigUnit_GetLengthControlRigSpline()
+	{
+		Length = 0.f;
+	}
+
+	/** Execute logic for this rig unit */
+	RIGVM_METHOD()
+	virtual void Execute(const FRigUnitContext& Context) override;
+
+	UPROPERTY(meta = (Input))
+	FControlRigSpline Spline;
+
+	UPROPERTY(meta = (Output))
+	float Length;
+};
