ControlRigSpline: Fix missing copyright boilerplate again

diff --git a/Engine/Plugins/Experimental/ControlRigSpline/Source/ControlRigSpline/Private/ControlRigSplineTypes.cpp b/Engine/Plugins/Experimental/ControlRigSpline/Source/ControlRigSpline/Private/ControlRigSplineTypes.cpp
--- a/Engine/Plugins/Experimental/ControlRigSpline/Source/ControlRigSpline/Private/ControlRigSplineTypes.cpp
+++ b/Engine/Plugins/Experimental/ControlRigSpline/Source/ControlRigSpline/Private/ControlRigSplineTypes.cpp
@@ -1,3 +1,5 @@
+// Copyright Epic Games, Inc. All Rights Reserved.
+
 #include "ControlRigSplineTypes.h"
 
 void FControlRigSpline::SetControlPoints(const TArray<FVector>& InPoints, const ESplineType SplineMode, const int32 SamplesPerSegment)
