Control Rig Spline: Edit CVs on a referenced curve

diff --git a/Engine/Plugins/Experimental/ControlRigSpline/Source/ControlRigSpline/Private/ControlRigSplineUnits.cpp b/Engine/Plugins/Experimental/ControlRigSpline/Source/ControlRigSpline/Private/ControlRigSplineUnits.cpp
--- a/Engine/Plugins/Experimental/ControlRigSpline/Source/ControlRigSpline/Private/ControlRigSplineUnits.cpp
+++ b/Engine/Plugins/Experimental/ControlRigSpline/Source/ControlRigSpline/Private/ControlRigSplineUnits.cpp
@@ -24,6 +24,36 @@ FRigUnit_ControlRigSplineFromPoints_Execute()
 	}
 }
 
+FRigUnit_SetSplinePoints_Execute()
+{
+	if (!Spline.SplineData.IsValid())
+	{
+		UE_LOG(LogControlRig, Error, TEXT("Invalid input spline."));
+		return;
+	}
+	
+	if (Points.Num() != Spline.SplineData->ControlPoints.Num())
+	{
+		UE_LOG(LogControlRig, Error, TEXT("Number of input points does not match the number of point in the spline."));
+		return;
+	}
+
+	switch (Context.State)
+	{
+		case EControlRigState::Init:
+		case EControlRigState::Update:
+		{
+			Spline.SetControlPoints(Points, Spline.SplineData->SplineMode, Spline.SplineData->SamplesPerSegment);
+			break;
+		}
+		default:
+		{
+			checkNoEntry(); // Execute is only defined for Init and Update
+			break;
+		}
+	}
+}
+
 FRigUnit_PositionFromControlRigSpline_Execute()
 {
 	switch (Context.State)
diff --git a/Engine/Plugins/Experimental/ControlRigSpline/Source/ControlRigSpline/Public/ControlRigSplineUnits.h b/Engine/Plugins/Experimental/ControlRigSpline/Source/ControlRigSpline/Public/ControlRigSplineUnits.h
--- a/Engine/Plugins/Experimental/ControlRigSpline/Source/ControlRigSpline/Public/ControlRigSplineUnits.h
+++ b/Engine/Plugins/Experimental/ControlRigSpline/Source/ControlRigSpline/Public/ControlRigSplineUnits.h
@@ -41,6 +41,27 @@ struct CONTROLRIGSPLINE_API FRigUnit_ControlRigSplineFromPoints : public FRigUni
 	FControlRigSpline Spline;
 };
 
+USTRUCT(meta = (DisplayName = "Set Spline Points", Category = "Control Rig"))
+struct CONTROLRIGSPLINE_API FRigUnit_SetSplinePoints : public FRigUnitMutable
+{
+	GENERATED_BODY()
+
+	FRigUnit_SetSplinePoints()
+	{
+		
+	}
+
+	/** Execute logic for this rig unit */
+	RIGVM_METHOD()
+	virtual void Execute(const FRigUnitContext& Context) override;
+
+	UPROPERTY(meta = (Input))
+	TArray<FVector> Points;
+
+	UPROPERTY(meta = (Input, Output))
+	FControlRigSpline Spline;
+};
+
 USTRUCT(meta = (DisplayName = "Position From Spline", Keywords="Point From Spline", Category = "Control Rig"))
 struct CONTROLRIGSPLINE_API FRigUnit_PositionFromControlRigSpline : public FRigUnit_ControlRigSplineBase
 {
