Control Rig Spline: Fit chain to Spline

diff --git a/Engine/Plugins/Experimental/ControlRigSpline/Source/ControlRigSpline/Private/ControlRigSplineUnits.cpp b/Engine/Plugins/Experimental/ControlRigSpline/Source/ControlRigSpline/Private/ControlRigSplineUnits.cpp
--- a/Engine/Plugins/Experimental/ControlRigSpline/Source/ControlRigSpline/Private/ControlRigSplineUnits.cpp
+++ b/Engine/Plugins/Experimental/ControlRigSpline/Source/ControlRigSpline/Private/ControlRigSplineUnits.cpp
@@ -155,4 +155,411 @@ FRigUnit_GetLengthControlRigSpline_Execute()
 			break;
 		}
 	}
-}
\ No newline at end of file
+}
+
+FRigUnit_FitChainToSplineCurve_Execute()
+{
+    DECLARE_SCOPE_HIERARCHICAL_COUNTER_RIGUNIT()
+
+	URigHierarchy* Hierarchy = ExecuteContext.Hierarchy;
+	if (Hierarchy == nullptr)
+	{
+		return;
+	}
+
+	if (!Spline.SplineData.IsValid())
+	{
+		return;
+	}
+
+	float& ChainLength = WorkData.ChainLength;
+	TArray<FVector>& ItemPositions = WorkData.ItemPositions;
+	TArray<float>& ItemSegments = WorkData.ItemSegments;
+	TArray<FVector>& CurvePositions = WorkData.CurvePositions;
+	TArray<float>& CurveSegments = WorkData.CurveSegments;
+	TArray<FCachedRigElement>& CachedItems = WorkData.CachedItems;
+	TArray<int32>& ItemRotationA = WorkData.ItemRotationA;
+	TArray<int32>& ItemRotationB = WorkData.ItemRotationB;
+	TArray<float>& ItemRotationT = WorkData.ItemRotationT;
+	TArray<FTransform>& ItemLocalTransforms = WorkData.ItemLocalTransforms;
+
+	if (Context.State == EControlRigState::Init)
+	{
+		CachedItems.Reset();
+		return;
+	}
+
+	if(CachedItems.Num() == 0 && Items.Num() > 1)
+	{
+		CurvePositions.Reset();
+		ItemPositions.Reset();
+		ItemSegments.Reset();
+		CurveSegments.Reset();
+		ItemRotationA.Reset();
+		ItemRotationB.Reset();
+		ItemRotationT.Reset();
+		ItemLocalTransforms.Reset();
+		ChainLength = 0.f;
+
+		for (FRigElementKey Item : Items)
+		{
+			CachedItems.Add(FCachedRigElement(Item, Hierarchy));
+		}
+
+		if (CachedItems.Num() < 2)
+		{
+			UE_CONTROLRIG_RIGUNIT_REPORT_WARNING(TEXT("Didn't find enough bones. You need at least two in the chain!"));
+			return;
+		}
+
+		ItemPositions.SetNumZeroed(CachedItems.Num());
+		ItemSegments.SetNumZeroed(CachedItems.Num());
+		ItemSegments[0] = 0;
+		for (int32 Index = 1; Index < CachedItems.Num(); Index++)
+		{
+			FVector A = Hierarchy->GetGlobalTransform(CachedItems[Index - 1]).GetLocation();
+			FVector B = Hierarchy->GetGlobalTransform(CachedItems[Index]).GetLocation();
+			ItemSegments[Index] = (A - B).Size();
+			ChainLength += ItemSegments[Index];
+		}
+
+		ItemRotationA.SetNumZeroed(CachedItems.Num());
+		ItemRotationB.SetNumZeroed(CachedItems.Num());
+		ItemRotationT.SetNumZeroed(CachedItems.Num());
+		ItemLocalTransforms.SetNumZeroed(CachedItems.Num());
+
+		if (Rotations.Num() > 1)
+		{
+			TArray<float> RotationRatios;
+			TArray<int32> RotationIndices;
+
+			for (const FRigUnit_FitChainToCurve_Rotation& Rotation : Rotations)
+			{
+				RotationIndices.Add(RotationIndices.Num());
+				RotationRatios.Add(FMath::Clamp<float>(Rotation.Ratio, 0.f, 1.f));
+			}
+
+			TLess<> Predicate;
+			auto Projection = [RotationRatios](int32 Val) -> float
+			{
+				return RotationRatios[Val];
+			};
+			Algo::SortBy(RotationIndices, Projection, Predicate);
+
+			for (int32 Index = 0; Index < CachedItems.Num(); Index++)
+			{
+				float T = 0.f;
+				if (CachedItems.Num() > 1)
+				{
+					T = float(Index) / float(CachedItems.Num() - 1);
+				}
+
+				if (T <= RotationRatios[RotationIndices[0]])
+				{
+					ItemRotationA[Index] = ItemRotationB[Index] = RotationIndices[0];
+					ItemRotationT[Index] = 0.f;
+				}
+				else if (T >= RotationRatios[RotationIndices.Last()])
+				{
+					ItemRotationA[Index] = ItemRotationB[Index] = RotationIndices.Last();
+					ItemRotationT[Index] = 0.f;
+				}
+				else
+				{
+					for (int32 RotationIndex = 1; RotationIndex < RotationIndices.Num(); RotationIndex++)
+					{
+						int32 A = RotationIndices[RotationIndex - 1];
+						int32 B = RotationIndices[RotationIndex];
+
+						if (FMath::IsNearlyEqual(Rotations[A].Ratio, T))
+						{
+							ItemRotationA[Index] = ItemRotationB[Index] = A;
+							ItemRotationT[Index] = 0.f;
+							break;
+						}
+						else if (FMath::IsNearlyEqual(Rotations[B].Ratio, T))
+						{
+							ItemRotationA[Index] = ItemRotationB[Index] = B;
+							ItemRotationT[Index] = 0.f;
+							break;
+						}
+						else if (Rotations[B].Ratio > T)
+						{
+							if (FMath::IsNearlyEqual(RotationRatios[A], RotationRatios[B]))
+							{
+								ItemRotationA[Index] = ItemRotationB[Index] = A;
+								ItemRotationT[Index] = 0.f;
+							}
+							else
+							{
+								ItemRotationA[Index] = A;
+								ItemRotationB[Index] = B;
+								ItemRotationT[Index] = (T - RotationRatios[A]) / (RotationRatios[B] - RotationRatios[A]);
+								ItemRotationT[Index] = FControlRigMathLibrary::EaseFloat(ItemRotationT[Index], RotationEaseType);
+							}
+							break;
+						}
+					}
+				}
+			}
+ 		}
+	}
+
+	if (CachedItems.Num() < 2)
+	{
+		return;
+	}
+
+	int32 Samples = FMath::Clamp<int32>(SamplingPrecision, 4, 64);
+	if (CurvePositions.Num() != Samples)
+	{
+		CurvePositions.SetNum(Samples + 1);
+		CurveSegments.SetNum(Samples + 1);
+	}
+
+	FVector StartTangent = FVector::ZeroVector;
+	FVector EndTangent = FVector::ZeroVector;
+
+	float CurveLength = 0.f;
+	for (int32 SampleIndex = 0; SampleIndex < Samples; SampleIndex++)
+	{
+		float T = float(SampleIndex) / float(Samples - 1);
+		T = FMath::Lerp<float>(Minimum, Maximum, T);
+
+		CurvePositions[SampleIndex] = Spline.PositionAtParam(T);
+		FVector Tangent = Spline.TangentAtParam(T);
+		if (SampleIndex == 0)
+		{
+			StartTangent = Tangent;
+		}
+		else if (SampleIndex == Samples - 1)
+		{
+			EndTangent = Tangent;
+		}
+
+		if (SampleIndex > 0)
+		{
+			CurveSegments[SampleIndex] = (CurvePositions[SampleIndex] - CurvePositions[SampleIndex - 1]).Size();
+			CurveLength += CurveSegments[SampleIndex];
+		}
+		else
+		{
+			CurveSegments[SampleIndex] = 0.f;
+		}
+	}
+
+	CurvePositions[Samples] = CurvePositions[Samples - 1] + EndTangent * ChainLength;
+	CurveSegments[Samples] = ChainLength;
+
+	if (ChainLength < SMALL_NUMBER)
+	{
+		UE_CONTROLRIG_RIGUNIT_REPORT_WARNING(TEXT("The chain has no length - all of the bones are in the sample place!"));
+		return;
+	}
+
+	if (CurveLength < SMALL_NUMBER)
+	{
+		UE_CONTROLRIG_RIGUNIT_REPORT_WARNING(TEXT("The curve has no length - all of the points are in the sample place!"));
+		return;
+	}
+
+	int32 CurveIndex = 1;
+	ItemPositions[0] = CurvePositions[0];
+
+	for (int32 Index = 1; Index < CachedItems.Num(); Index++)
+	{
+		const FVector& LastPosition = ItemPositions[Index - 1];
+
+		float BoneLength = ItemSegments[Index];
+		switch (Alignment)
+		{
+			case EControlRigCurveAlignment::Front:
+			{
+				break;
+			}
+			case EControlRigCurveAlignment::Stretched:
+			{
+				BoneLength = BoneLength * CurveLength / ChainLength;
+				break;
+			}
+		}
+
+		FVector A = FVector::ZeroVector;
+		FVector B = FVector::ZeroVector;
+		A = CurvePositions[CurveIndex - 1];
+		B = CurvePositions[CurveIndex];
+
+		float DistanceA = (LastPosition - A).Size();
+		float DistanceB = (LastPosition - B).Size();
+		
+		if (DistanceB > BoneLength)
+		{
+			float Ratio = BoneLength / DistanceB;
+			ItemPositions[Index] = FMath::Lerp<FVector>(LastPosition, B, Ratio);
+			continue;
+		}
+
+		while (CurveIndex < CurvePositions.Num() - 1)
+		{
+			CurveIndex++;
+			A = B;
+			B = CurvePositions[CurveIndex];
+			DistanceA = DistanceB;
+			DistanceB = (B - LastPosition).Size();
+
+			if ((DistanceA < BoneLength) != (DistanceB < BoneLength))
+			{
+				break;
+			}
+		}
+
+		if (DistanceB < DistanceA)
+		{
+			FVector TempV = A;
+			A = B;
+			B = TempV;
+			float TempF = DistanceA;
+			DistanceA = DistanceB;
+			DistanceB = TempF;
+		}
+
+		if (FMath::IsNearlyEqual(DistanceA, DistanceB))
+		{
+			ItemPositions[Index] = A;
+			continue;
+		}
+
+		float Ratio = (BoneLength - DistanceA) / (DistanceB - DistanceA);
+		ItemPositions[Index] = FMath::Lerp<FVector>(A, B, Ratio);
+	}
+
+	TArray<FTransform> BoneGlobalTransforms;
+	if (Weight < 1.f - SMALL_NUMBER)
+	{
+		for (int32 Index = 0; Index < CachedItems.Num(); Index++)
+		{
+			BoneGlobalTransforms.Add(Hierarchy->GetGlobalTransform(CachedItems[Index]));
+		}
+	}
+
+	for (int32 Index = 0; Index < CachedItems.Num(); Index++)
+	{
+		FTransform Transform = Hierarchy->GetGlobalTransform(CachedItems[Index]);
+
+		Transform.SetTranslation(ItemPositions[Index]);
+
+		FVector Target = FVector::ZeroVector;
+		if (Index < CachedItems.Num() - 1)
+		{
+			Target = ItemPositions[Index + 1] - ItemPositions[Index];
+		}
+		else
+		{
+			Target = ItemPositions.Last() - ItemPositions[ItemPositions.Num() - 2];
+		}
+
+		if (!Target.IsNearlyZero() && !PrimaryAxis.IsNearlyZero())
+		{
+			Target = Target.GetSafeNormal();
+			FVector Axis = Transform.TransformVectorNoScale(PrimaryAxis).GetSafeNormal();
+			FQuat Rotation = FQuat::FindBetweenNormals(Axis, Target);
+			Transform.SetRotation((Rotation * Transform.GetRotation()).GetNormalized());
+		}
+
+		Target = PoleVectorPosition - ItemPositions[Index];
+		if (!SecondaryAxis.IsNearlyZero())
+		{
+			if (!PrimaryAxis.IsNearlyZero())
+			{
+				FVector Axis = Transform.TransformVectorNoScale(PrimaryAxis).GetSafeNormal();
+				Target = Target - FVector::DotProduct(Target, Axis) * Axis;
+			}
+
+			if (!Target.IsNearlyZero() && !SecondaryAxis.IsNearlyZero())
+			{
+				Target = Target.GetSafeNormal();
+				FVector Axis = Transform.TransformVectorNoScale(SecondaryAxis).GetSafeNormal();
+				FQuat Rotation = FQuat::FindBetweenNormals(Axis, Target);
+				Transform.SetRotation((Rotation * Transform.GetRotation()).GetNormalized());
+			}
+		}
+
+		if (Weight >= 1.f - SMALL_NUMBER)
+		{
+			Hierarchy->SetGlobalTransform(CachedItems[Index], Transform, bPropagateToChildren && Rotations.Num() == 0);
+		}
+		else
+		{
+			Transform = FControlRigMathLibrary::LerpTransform(BoneGlobalTransforms[Index], Transform, FMath::Clamp<float>(Weight, 0.f, 1.f));
+			Hierarchy->SetGlobalTransform(CachedItems[Index], Transform, bPropagateToChildren && Rotations.Num() == 0);
+		}
+	}
+
+	if (Rotations.Num() > 0)
+	{
+		FTransform BaseTransform = FTransform::Identity;
+		FRigElementKey ParentKey = Hierarchy->GetFirstParent(CachedItems[0].GetKey());
+		if (ParentKey.IsValid())
+		{
+			BaseTransform = Hierarchy->GetGlobalTransform(ParentKey);
+		}
+
+		for (int32 Index = 0; Index < CachedItems.Num(); Index++)
+		{
+			ItemLocalTransforms[Index] = Hierarchy->GetLocalTransform(CachedItems[Index]);
+		}
+
+		for (int32 Index = 0; Index < CachedItems.Num(); Index++)
+		{
+			if (ItemRotationA[Index] >= Rotations.Num() ||
+				ItemRotationB[Index] >= Rotations.Num())
+			{
+				continue;
+			}
+
+			FQuat Rotation = Rotations[ItemRotationA[Index]].Rotation;
+			FQuat RotationB = Rotations[ItemRotationB[Index]].Rotation;
+			if (ItemRotationA[Index] != ItemRotationB[Index])
+			{
+				if (ItemRotationT[Index] > 1.f - SMALL_NUMBER)
+				{
+					Rotation = RotationB;
+				}
+				else if (ItemRotationT[Index] > SMALL_NUMBER)
+				{
+					Rotation = FQuat::Slerp(Rotation, RotationB, ItemRotationT[Index]).GetNormalized();
+				}
+			}
+
+			BaseTransform = ItemLocalTransforms[Index] * BaseTransform;
+			Rotation = FQuat::Slerp(BaseTransform.GetRotation(), BaseTransform.GetRotation() * Rotation, FMath::Clamp<float>(Weight, 0.f, 1.f));
+			BaseTransform.SetRotation(Rotation);
+			Hierarchy->SetGlobalTransform(CachedItems[Index], BaseTransform, bPropagateToChildren);
+		}
+	}
+
+	if (Context.DrawInterface != nullptr && DebugSettings.bEnabled)
+	{
+		int32 Count = 64;
+		FControlRigDrawInstruction Instruction(EControlRigDrawSettings::LineStrip, DebugSettings.CurveColor, DebugSettings.Scale, DebugSettings.WorldOffset);
+		Instruction.Positions.SetNumUninitialized(Count);
+
+		float T = 0;
+		float Step = 1.f / float(Count-1);
+		for(int32 Index=0; Index<Count; ++Index)
+		{
+			// Evaluate at T
+			Instruction.Positions[Index] = Spline.PositionAtParam(T);
+			T += Step;
+		}
+		Context.DrawInterface->Instructions.Add(Instruction);
+
+		for (auto Point : Spline.SplineData->ControlPoints)
+		{
+			Context.DrawInterface->DrawPoint(DebugSettings.WorldOffset, Point, DebugSettings.Scale * 6, DebugSettings.CurveColor);
+		}
+
+		Context.DrawInterface->DrawLineStrip(DebugSettings.WorldOffset, CurvePositions, DebugSettings.SegmentsColor, DebugSettings.Scale);
+		Context.DrawInterface->DrawPoints(DebugSettings.WorldOffset, CurvePositions, DebugSettings.Scale * 4.f, DebugSettings.SegmentsColor);
+	}
+}
diff --git a/Engine/Plugins/Experimental/ControlRigSpline/Source/ControlRigSpline/Public/ControlRigSplineUnits.h b/Engine/Plugins/Experimental/ControlRigSpline/Source/ControlRigSpline/Public/ControlRigSplineUnits.h
--- a/Engine/Plugins/Experimental/ControlRigSpline/Source/ControlRigSpline/Public/ControlRigSplineUnits.h
+++ b/Engine/Plugins/Experimental/ControlRigSpline/Source/ControlRigSpline/Public/ControlRigSplineUnits.h
@@ -5,6 +5,7 @@
 #include "Units/RigUnit.h"
 #include "ControlRig.h"
 #include "ControlRigSplineTypes.h"
+#include "ControlRig/Private/Units/Highlevel/Hierarchy/RigUnit_FitChainToCurve.h"
 #include "ControlRigSplineUnits.generated.h"
 
 USTRUCT(meta = (Abstract, NodeColor = "0.3 0.1 0.1"))
@@ -169,3 +170,119 @@ struct CONTROLRIGSPLINE_API FRigUnit_GetLengthControlRigSpline : public FRigUnit
 	UPROPERTY(meta = (Output))
 	float Length;
 };
+
+/**
+ * Fits a given chain to a spline curve.
+ * Additionally provides rotational control matching the features of the Distribute Rotation node.
+ */
+USTRUCT(meta=(DisplayName="Fit Chain on Spline Curve", Category="Hierarchy", Keywords="Fit,Resample,Spline"))
+struct CONTROLRIGSPLINE_API FRigUnit_FitChainToSplineCurve : public FRigUnit_HighlevelBaseMutable
+{
+	GENERATED_BODY()
+
+	FRigUnit_FitChainToSplineCurve()
+	{
+		Alignment = EControlRigCurveAlignment::Stretched;
+		Minimum = 0.f;
+		Maximum = 1.f;
+		SamplingPrecision = 12;
+		PrimaryAxis = FVector(1.f, 0.f, 0.f);
+		SecondaryAxis = FVector(0.f, 0.f, 0.f);
+		PoleVectorPosition = FVector::ZeroVector;
+		RotationEaseType = EControlRigAnimEasingType::Linear;
+		Weight = 1.f;
+		bPropagateToChildren = true;
+		DebugSettings = FRigUnit_FitChainToCurve_DebugSettings();
+	}
+
+	RIGVM_METHOD()
+	virtual void Execute(const FRigUnitContext& Context) override;
+
+	/** 
+	 * The items to align
+	 */
+	UPROPERTY(meta = (Input))
+	FRigElementKeyCollection Items;
+
+	/** 
+	 * The curve to align to
+	 */
+	UPROPERTY(meta = (Input))
+	FControlRigSpline Spline;
+
+	/** 
+	 * Specifies how to align the chain on the curve
+	 */
+	UPROPERTY(meta = (Input, Constant))
+	EControlRigCurveAlignment Alignment;
+
+	/** 
+	 * The minimum U value to use on the curve
+	 */
+	UPROPERTY(meta = (Input, Constant))
+	float Minimum;
+
+	/** 
+	 * The maximum U value to use on the curve
+	 */
+	UPROPERTY(meta = (Input, Constant))
+	float Maximum;
+
+	/**
+	 * The number of samples to use on the curve. Clamped at 64.
+	 */
+	UPROPERTY(meta = (Input, Constant))
+	int32 SamplingPrecision;
+
+	/**
+	 * The major axis being aligned - along the bone
+	 */
+	UPROPERTY(meta = (Input))
+	FVector PrimaryAxis;
+
+	/**
+	 * The minor axis being aligned - towards the pole vector.
+	 * You can use (0.0, 0.0, 0.0) to disable it.
+	 */
+	UPROPERTY(meta = (Input))
+	FVector SecondaryAxis;
+
+	/**
+	 * The the position of the pole vector used for aligning the secondary axis.
+	 * Only has an effect if the secondary axis is set.
+	 */
+	UPROPERTY(meta = (Input))
+	FVector PoleVectorPosition;
+
+	/** 
+	 * The list of rotations to be applied along the curve
+	 */
+	UPROPERTY(meta = (Input))
+	TArray<FRigUnit_FitChainToCurve_Rotation> Rotations;
+
+	/**
+	 * The easing to use between to rotations.
+	 */
+	UPROPERTY(meta = (Input, Constant))
+	EControlRigAnimEasingType RotationEaseType;
+
+	/**
+	 * The weight of the solver - how much the rotation should be applied
+	 */
+	UPROPERTY(meta = (Input))
+	float Weight;
+
+	/**
+	 * If set to true all of the global transforms of the children
+	 * of this bone will be recalculated based on their local transforms.
+	 * Note: This is computationally more expensive than turning it off.
+	 */
+	UPROPERTY(meta = (Input, Constant))
+	bool bPropagateToChildren;
+
+	UPROPERTY(meta = (Input, DetailsOnly))
+	FRigUnit_FitChainToCurve_DebugSettings DebugSettings;
+
+	UPROPERTY(transient)
+	FRigUnit_FitChainToCurve_WorkData WorkData;
+};
\ No newline at end of file
