Control Rig Spline: Enable stretch or fix length

diff --git a/Engine/Plugins/Experimental/ControlRigSpline/Source/ControlRigSpline/Private/ControlRigSplineTypes.cpp b/Engine/Plugins/Experimental/ControlRigSpline/Source/ControlRigSpline/Private/ControlRigSplineTypes.cpp
--- a/Engine/Plugins/Experimental/ControlRigSpline/Source/ControlRigSpline/Private/ControlRigSplineTypes.cpp
+++ b/Engine/Plugins/Experimental/ControlRigSpline/Source/ControlRigSpline/Private/ControlRigSplineTypes.cpp
@@ -2,7 +2,7 @@
 
 #include "ControlRigSplineTypes.h"
 
-void FControlRigSpline::SetControlPoints(const TArrayView<const FVector>& InPoints, const ESplineType SplineMode, const int32 SamplesPerSegment)
+void FControlRigSpline::SetControlPoints(const TArrayView<const FVector>& InPoints, const ESplineType SplineMode, const int32 SamplesPerSegment, const float Compression, const float Stretch)
 {
 	const int32 ControlPointsCount = InPoints.Num();
 	if (ControlPointsCount < 4)
@@ -18,7 +18,9 @@ void FControlRigSpline::SetControlPoints(const TArrayView<const FVector>& InPoin
 	bool bControlPointsChanged = InPoints != SplineData->ControlPoints;
 	bool bSplineModeChanged = SplineMode != SplineData->SplineMode;
 	bool bSamplesCountChanged = SamplesPerSegment != SplineData->SamplesPerSegment;
-	if (!bSplineModeChanged && !bControlPointsChanged && !bSamplesCountChanged)
+	bool bNumControlPointsChanged = SplineData->ControlPoints.Num() != ControlPointsCount;
+	bool bStretchChanged = Stretch != SplineData->Stretch || Compression != SplineData->Compression;
+	if (!bSplineModeChanged && !bControlPointsChanged && !bSamplesCountChanged && !bStretchChanged)
 	{
 		return;
 	}
@@ -26,6 +28,8 @@ void FControlRigSpline::SetControlPoints(const TArrayView<const FVector>& InPoin
 	SplineData->ControlPoints = InPoints;
 	SplineData->SplineMode = SplineMode;
 	SplineData->SamplesPerSegment = SamplesPerSegment;
+	SplineData->Compression = Compression;
+	SplineData->Stretch = Stretch;
 
 	// If we need to update the spline because the controls points have changed, or the spline mode has changed
 	if (bControlPointsChanged || bSplineModeChanged)
@@ -34,7 +38,7 @@ void FControlRigSpline::SetControlPoints(const TArrayView<const FVector>& InPoin
 		{
 			case ESplineType::BSpline:
 			{
-				if (bSplineModeChanged || SplineData->Spline.numControlPoints() != ControlPointsCount)
+				if (bSplineModeChanged || bNumControlPointsChanged)
 				{
 					SplineData->Spline = tinyspline::BSpline(ControlPointsCount, 3);
 				}
@@ -54,41 +58,28 @@ void FControlRigSpline::SetControlPoints(const TArrayView<const FVector>& InPoin
 			}
 			case ESplineType::Hermite:
 			{
-				if (bSplineModeChanged || SplineData->Spline.numControlPoints() != ControlPointsCount)
+				// With this spline type, we cannot just move the contorl points around, since the interpolation creates
+				// new control points from the ones provided. The interpolation allocates a new spline each time
+				// its called.
+						
+				// Create the control points array for the interpolation.
+				// The first dimension acts as the query parameter, the rest of the dimensions are the XYZ of the control points
+				TArray<float> ControlPointsArray;
+				ControlPointsArray.SetNumUninitialized(ControlPointsCount * 4, false);
+	
+				for (int32 i = 0; i < ControlPointsCount; ++i)
 				{
-					// Create a new BSpline from the control points
-					// The first dimension acts as the query parameter, the rest of the dimensions are the XYZ of the control points
-					TArray<float> ControlPointsArray;
-					ControlPointsArray.SetNumUninitialized(ControlPointsCount * 4, false);
-				
-					for (int32 i = 0; i < ControlPointsCount; ++i)
-					{
-						// The control point with weight 0 is the second point
-						// The control point with weight 1 is the second to last point
-						// This way the first and last point act as tangents
-						ControlPointsArray[i * 4 + 0] = (i - 1) / (float)(ControlPointsCount - 3);
-
-						ControlPointsArray[i * 4 + 1] = (float)InPoints[i][0];
-						ControlPointsArray[i * 4 + 2] = (float)InPoints[i][1];
-						ControlPointsArray[i * 4 + 3] = (float)InPoints[i][2];
-					}
+					ControlPointsArray[i * 4 + 0] = i / (float)(ControlPointsCount - 1);
 
-					SplineData->Spline = tinyspline::BSpline(ControlPointsCount, 4);
-					ts_bspline_interpolate_cubic_natural(ControlPointsArray.GetData(), ControlPointsCount, 4, SplineData->Spline.data(), nullptr);
-				}
-				else
-				{
-					// We just need to update the positions of the control points
-					float* Ctrlp = nullptr;
-					ts_bspline_control_point_at(SplineData->Spline.data(), 0, &Ctrlp, nullptr);
-					for (int32 i = 0; i < ControlPointsCount; ++i)
-					{
-						Ctrlp[i * 4 + 1] = (float)InPoints[i][0];
-						Ctrlp[i * 4 + 2] = (float)InPoints[i][1];
-						Ctrlp[i * 4 + 3] = (float)InPoints[i][2];
-					}
+					ControlPointsArray[i * 4 + 1] = (float)InPoints[i][0];
+					ControlPointsArray[i * 4 + 2] = (float)InPoints[i][1];
+					ControlPointsArray[i * 4 + 3] = (float)InPoints[i][2];
 				}
 
+				// Create a new spline to store the result of the interpolation
+				ts_bspline_free(SplineData->Spline.data());					
+				ts_bspline_interpolate_cubic_natural(ControlPointsArray.GetData(), ControlPointsCount, 4, SplineData->Spline.data(), nullptr);
+				
 				break;
 			}
 			default:
@@ -100,7 +91,7 @@ void FControlRigSpline::SetControlPoints(const TArrayView<const FVector>& InPoin
 	}
 
 	// If curve has changed, or sample count has changed, recompute the cache
-	if (bControlPointsChanged || bSplineModeChanged || bSamplesCountChanged)
+	if (bControlPointsChanged || bSplineModeChanged || bSamplesCountChanged || bStretchChanged)
 	{
 		switch (SplineMode)
 		{
@@ -119,6 +110,33 @@ void FControlRigSpline::SetControlPoints(const TArrayView<const FVector>& InPoin
 					}
 				}
 
+				// Correct length of samples
+				if (!bSplineModeChanged && !bSamplesCountChanged && !bNumControlPointsChanged)
+				{
+					if (SplineData->InitialLengths.Num() > 0)
+					{
+						TArray<FVector3f> SamplesBeforeCorrect = SplineData->SamplesArray;
+						for (int32 i = 0; i < ControlPointsCount - 1; ++i)
+						{
+							for (int32 j = (i == 0) ? 1 : 0; j < SamplesPerSegment; ++j)
+							{
+								// Get direction from samples before correction
+								FVector3f Dir = SamplesBeforeCorrect[i * SamplesPerSegment + j] - SamplesBeforeCorrect[i * SamplesPerSegment + j - 1];
+								Dir.Normalize();
+
+								float InitialLength = SplineData->InitialLengths[i * SamplesPerSegment + j - 1];
+								// Current length as the projection on the Dir vector (might be negative)
+								float CurrentLength = (SplineData->SamplesArray[i * SamplesPerSegment + j] - SplineData->SamplesArray[i * SamplesPerSegment + j - 1]).Dot(Dir);
+								float FixedLength = FMath::Clamp(CurrentLength,
+									(Compression > 0.f) ? InitialLength * Compression : CurrentLength,
+									(Stretch > 0.f) ? InitialLength * Stretch : CurrentLength);
+
+								SplineData->SamplesArray[i * SamplesPerSegment + j] = SplineData->SamplesArray[i * SamplesPerSegment + j - 1] + Dir * FixedLength;
+							}
+						}
+					}
+				}
+
 				// tinySpline will allocate the samples array, but does not free that memory. We need to take care of that.
 				free(SamplesPtr);
 
@@ -130,18 +148,45 @@ void FControlRigSpline::SetControlPoints(const TArrayView<const FVector>& InPoin
 				float* SamplesPtr = nullptr;
 				size_t ActualSamplesPerSegment = 0;
 				ts_bspline_sample(SplineData->Spline.data(), (ControlPointsCount - 1) * SamplesPerSegment, &SamplesPtr, &ActualSamplesPerSegment, nullptr);
-				SplineData->SamplesArray.SetNumUninitialized((ControlPointsCount - 3) * SamplesPerSegment);
-				for (int32 i = 0; i < ControlPointsCount - 3; ++i)
+				SplineData->SamplesArray.SetNumUninitialized((ControlPointsCount - 1) * SamplesPerSegment);
+				for (int32 i = 0; i < ControlPointsCount - 1; ++i)
 				{
 					for (int32 j = 0; j < SamplesPerSegment; ++j)
 					{
-						memcpy(&SplineData->SamplesArray[i * SamplesPerSegment + j].X, SamplesPtr + ((i + 1) * SamplesPerSegment + j) * 4 + 1, sizeof(float) * 3);
+						memcpy(&SplineData->SamplesArray[i * SamplesPerSegment + j].X, SamplesPtr + (i * SamplesPerSegment + j) * 4 + 1, sizeof(float) * 3);
 					}
 				}
 
 				// tinySpline will allocate the samples array, but does not free that memory. We need to take care of that.
 				free(SamplesPtr);
 
+				// Correct length of samples
+				if (!bSplineModeChanged && !bSamplesCountChanged && !bNumControlPointsChanged)
+				{
+					if (SplineData->InitialLengths.Num() > 0)
+					{
+						TArray<FVector3f> SamplesBeforeCorrect = SplineData->SamplesArray;
+						for (int32 i = 0; i < ControlPointsCount - 1; ++i)
+						{
+							for (int32 j = (i == 0) ? 1 : 0; j < SamplesPerSegment; ++j)
+							{
+								// Get direction from samples before correction
+								FVector3f Dir = SamplesBeforeCorrect[i * SamplesPerSegment + j] - SamplesBeforeCorrect[i * SamplesPerSegment + j - 1];
+								Dir.Normalize();
+
+								float InitialLength = SplineData->InitialLengths[i * SamplesPerSegment + j - 1];
+								// Current length as the projection on the Dir vector (might be negative)
+								float CurrentLength = (SplineData->SamplesArray[i * SamplesPerSegment + j] - SplineData->SamplesArray[i * SamplesPerSegment + j - 1]).Dot(Dir);
+								float FixedLength = FMath::Clamp(CurrentLength,
+                                	(Compression > 0.f) ? InitialLength * Compression : CurrentLength,
+                                	(Stretch > 0.f) ? InitialLength * Stretch : CurrentLength);
+
+								SplineData->SamplesArray[i * SamplesPerSegment + j] = SplineData->SamplesArray[i * SamplesPerSegment + j - 1] + Dir * FixedLength;
+							}
+						}
+					}
+				}
+
 				break;
 			}
 			default:
@@ -154,9 +199,18 @@ void FControlRigSpline::SetControlPoints(const TArrayView<const FVector>& InPoin
 		// Cache accumulated length at sample array
 		SplineData->AccumulatedLenth.SetNumUninitialized(SplineData->SamplesArray.Num(), false);
 		SplineData->AccumulatedLenth[0] = 0.f;
+		if (bSplineModeChanged || bSamplesCountChanged || bNumControlPointsChanged)
+		{
+			SplineData->InitialLengths.SetNumUninitialized(SplineData->SamplesArray.Num() - 1);
+		}
 		for (int32 i = 1; i < SplineData->SamplesArray.Num(); ++i)
 		{
-			SplineData->AccumulatedLenth[i] = SplineData->AccumulatedLenth[i - 1] + FVector::Distance(SplineData->SamplesArray[i - 1], SplineData->SamplesArray[i]);
+			float CurrentLength = FVector::Distance(SplineData->SamplesArray[i - 1], SplineData->SamplesArray[i]);
+			if (bSplineModeChanged || bSamplesCountChanged || bNumControlPointsChanged)
+			{
+				SplineData->InitialLengths[i-1] = CurrentLength;
+			}
+			SplineData->AccumulatedLenth[i] = SplineData->AccumulatedLenth[i - 1] + CurrentLength;
 		}
 	}	
 }
@@ -164,7 +218,7 @@ void FControlRigSpline::SetControlPoints(const TArrayView<const FVector>& InPoin
 int32 FControlRigSpline::GetControlPoints(TArray<FVector>& OutPoints) const
 {
 	const size_t Count = SplineData->Spline.numControlPoints();
-	OutPoints.Reset(Count);
+	OutPoints.SetNumUninitialized(Count);
 
 	float* Points;
 	ts_bspline_control_points(SplineData->Spline.data(), &Points, nullptr);
@@ -184,7 +238,7 @@ int32 FControlRigSpline::GetControlPoints(TArray<FVector>& OutPoints) const
 		{
 			for (int32 i = 0; i < Count; ++i)
 			{
-				int32 index = i * 3;
+				int32 index = i * 4;
 				OutPoints[i].Set(Points[index + 1], Points[index + 2], Points[index + 3]);
 			}
 			break;
diff --git a/Engine/Plugins/Experimental/ControlRigSpline/Source/ControlRigSpline/Private/ControlRigSplineUnits.cpp b/Engine/Plugins/Experimental/ControlRigSpline/Source/ControlRigSpline/Private/ControlRigSplineUnits.cpp
--- a/Engine/Plugins/Experimental/ControlRigSpline/Source/ControlRigSpline/Private/ControlRigSplineUnits.cpp
+++ b/Engine/Plugins/Experimental/ControlRigSpline/Source/ControlRigSpline/Private/ControlRigSplineUnits.cpp
@@ -14,7 +14,7 @@ FRigUnit_ControlRigSplineFromPoints_Execute()
 		case EControlRigState::Update:
 		{
 			const TArrayView<const FVector> PointsView(Points.GetData(), Points.Num());
-			Spline.SetControlPoints(PointsView, SplineMode, SamplesPerSegment);
+			Spline.SetControlPoints(PointsView, SplineMode, SamplesPerSegment, Compression, Stretch);
 			break;
 		}
 		default:
@@ -45,7 +45,7 @@ FRigUnit_SetSplinePoints_Execute()
 		case EControlRigState::Update:
 		{
 			const TArrayView<const FVector> PointsView(Points.GetData(), Points.Num());
-			Spline.SetControlPoints(PointsView, Spline.SplineData->SplineMode, Spline.SplineData->SamplesPerSegment);
+			Spline.SetControlPoints(PointsView, Spline.SplineData->SplineMode, Spline.SplineData->SamplesPerSegment, Spline.SplineData->Compression, Spline.SplineData->Stretch);
 			break;
 		}
 		default:
@@ -58,6 +58,11 @@ FRigUnit_SetSplinePoints_Execute()
 
 FRigUnit_PositionFromControlRigSpline_Execute()
 {
+	if (!Spline.SplineData.IsValid())
+	{
+		return;
+	}
+	
 	switch (Context.State)
 	{
 		case EControlRigState::Init:
@@ -76,6 +81,11 @@ FRigUnit_PositionFromControlRigSpline_Execute()
 
 FRigUnit_TransformFromControlRigSpline_Execute()
 {
+	if (!Spline.SplineData.IsValid())
+	{
+		return;
+	}
+	
 	switch (Context.State)
 	{
 		case EControlRigState::Init:
@@ -122,6 +132,11 @@ FRigUnit_DrawControlRigSpline_Execute()
 	{
 		return;
 	}
+	
+	if (!Spline.SplineData.IsValid())
+	{
+		return;
+	}
 
 	int32 Count = FMath::Clamp<int32>(Detail, 4, 64);
 	FControlRigDrawInstruction Instruction(EControlRigDrawSettings::LineStrip, Color, Thickness);
@@ -141,6 +156,12 @@ FRigUnit_DrawControlRigSpline_Execute()
 
 FRigUnit_GetLengthControlRigSpline_Execute()
 {
+	if (!Spline.SplineData.IsValid())
+	{
+		Length = 0;
+		return;
+	}
+	
 	switch (Context.State)
 	{
 		case EControlRigState::Init:
diff --git a/Engine/Plugins/Experimental/ControlRigSpline/Source/ControlRigSpline/Public/ControlRigSplineTypes.h b/Engine/Plugins/Experimental/ControlRigSpline/Source/ControlRigSpline/Public/ControlRigSplineTypes.h
--- a/Engine/Plugins/Experimental/ControlRigSpline/Source/ControlRigSpline/Public/ControlRigSplineTypes.h
+++ b/Engine/Plugins/Experimental/ControlRigSpline/Source/ControlRigSpline/Public/ControlRigSplineTypes.h
@@ -17,7 +17,7 @@ enum class ESplineType : uint8
 	BSpline,
 
 	/** Hermite: */
-	/** The curve will pass through the control points (except the first and last points, which are used as tangents) */
+	/** The curve will pass through the control points */
 	Hermite,
 
 	/** MAX - invalid */
@@ -41,12 +41,21 @@ struct CONTROLRIGSPLINE_API FControlRigSplineImpl
 	// The control points to construct the spline
 	TArray<FVector> ControlPoints;
 
+	// The initial lengths between samples
+	TArray<float> InitialLengths;
+
 	// The actual spline
 	tinyspline::BSpline Spline;
 
 	// Samples per segment, where segment is the portion between two control points
 	int32 SamplesPerSegment;
 
+	// The allowed length compression (1.f being do not allow compression`). If 0, no restriction wil be applied.
+	float Compression;
+
+	// The allowed length stretch (1.f being do not allow stretch`). If 0, no restriction wil be applied.
+	float Stretch;
+
 	// Positions along the "real" curve (no samples in the first and last segments of a hermite spline)
 	TArray<FVector> SamplesArray;
 
@@ -71,9 +80,12 @@ struct CONTROLRIGSPLINE_API FControlRigSpline
 	* depend on what is set in SplineMode.
 	*
 	* @param InPoints	The control points to set.
-	* @param forceRebuild	If true, will build the spline from scratch.
+	* @param SplineMode	The type of spline
+	* @param SamplesPerSegment The samples to cache for every segment defined between two control rig
+	* @param Compression The allowed length compression (1.f being do not allow compression`). If 0, no restriction wil be applied.
+	* @param Stretch The allowed length stretch (1.f being do not allow stretch). If 0, no restriction wil be applied.
 	*/
-	void SetControlPoints(const TArrayView<const FVector>& InPoints, const ESplineType SplineMode = ESplineType::BSpline, const int32 SamplesPerSegment = 16);
+	void SetControlPoints(const TArrayView<const FVector>& InPoints, const ESplineType SplineMode = ESplineType::BSpline, const int32 SamplesPerSegment = 16, const float Compression = 1.f, const float Stretch = 1.f);
 
 	/**
 	* Populates OutPoints and returns the number of control points in the spline.
diff --git a/Engine/Plugins/Experimental/ControlRigSpline/Source/ControlRigSpline/Public/ControlRigSplineUnits.h b/Engine/Plugins/Experimental/ControlRigSpline/Source/ControlRigSpline/Public/ControlRigSplineUnits.h
--- a/Engine/Plugins/Experimental/ControlRigSpline/Source/ControlRigSpline/Public/ControlRigSplineUnits.h
+++ b/Engine/Plugins/Experimental/ControlRigSpline/Source/ControlRigSpline/Public/ControlRigSplineUnits.h
@@ -23,6 +23,8 @@ struct CONTROLRIGSPLINE_API FRigUnit_ControlRigSplineFromPoints : public FRigUni
 	{
 		SplineMode = ESplineType::BSpline;
 		SamplesPerSegment = 16;
+		Compression = 1.f;
+		Stretch = 1.f;
 	}
 
 	/** Execute logic for this rig unit */
@@ -38,6 +40,12 @@ struct CONTROLRIGSPLINE_API FRigUnit_ControlRigSplineFromPoints : public FRigUni
 	UPROPERTY(meta = (Input))
 	int32 SamplesPerSegment;
 
+	UPROPERTY(meta = (Input))
+	float Compression;
+
+	UPROPERTY(meta = (Input))
+	float Stretch;
+
 	UPROPERTY(meta = (Output))
 	FControlRigSpline Spline;
 };
