Control Rig Spline: Fix compilation error in ControlRigSplineTypes.cpp

diff --git a/Engine/Plugins/Experimental/ControlRigSpline/Source/ControlRigSpline/Private/ControlRigSplineTypes.cpp b/Engine/Plugins/Experimental/ControlRigSpline/Source/ControlRigSpline/Private/ControlRigSplineTypes.cpp
--- a/Engine/Plugins/Experimental/ControlRigSpline/Source/ControlRigSpline/Private/ControlRigSplineTypes.cpp
+++ b/Engine/Plugins/Experimental/ControlRigSpline/Source/ControlRigSpline/Private/ControlRigSplineTypes.cpp
@@ -115,13 +115,13 @@ void FControlRigSpline::SetControlPoints(const TArrayView<const FVector>& InPoin
 				{
 					if (SplineData->InitialLengths.Num() > 0)
 					{
-						TArray<FVector3f> SamplesBeforeCorrect = SplineData->SamplesArray;
+						TArray<FVector> SamplesBeforeCorrect = SplineData->SamplesArray;
 						for (int32 i = 0; i < ControlPointsCount - 1; ++i)
 						{
 							for (int32 j = (i == 0) ? 1 : 0; j < SamplesPerSegment; ++j)
 							{
 								// Get direction from samples before correction
-								FVector3f Dir = SamplesBeforeCorrect[i * SamplesPerSegment + j] - SamplesBeforeCorrect[i * SamplesPerSegment + j - 1];
+								FVector Dir = SamplesBeforeCorrect[i * SamplesPerSegment + j] - SamplesBeforeCorrect[i * SamplesPerSegment + j - 1];
 								Dir.Normalize();
 
 								float InitialLength = SplineData->InitialLengths[i * SamplesPerSegment + j - 1];
@@ -165,13 +165,13 @@ void FControlRigSpline::SetControlPoints(const TArrayView<const FVector>& InPoin
 				{
 					if (SplineData->InitialLengths.Num() > 0)
 					{
-						TArray<FVector3f> SamplesBeforeCorrect = SplineData->SamplesArray;
+						TArray<FVector> SamplesBeforeCorrect = SplineData->SamplesArray;
 						for (int32 i = 0; i < ControlPointsCount - 1; ++i)
 						{
 							for (int32 j = (i == 0) ? 1 : 0; j < SamplesPerSegment; ++j)
 							{
 								// Get direction from samples before correction
-								FVector3f Dir = SamplesBeforeCorrect[i * SamplesPerSegment + j] - SamplesBeforeCorrect[i * SamplesPerSegment + j - 1];
+								FVector Dir = SamplesBeforeCorrect[i * SamplesPerSegment + j] - SamplesBeforeCorrect[i * SamplesPerSegment + j - 1];
 								Dir.Normalize();
 
 								float InitialLength = SplineData->InitialLengths[i * SamplesPerSegment + j - 1];
