Control Rig Spline: Add FitCurveToChain node

diff --git a/Engine/Plugins/Animation/ControlRig/Source/ControlRigEditor/Private/Editor/ControlRigEditor.cpp b/Engine/Plugins/Animation/ControlRig/Source/ControlRigEditor/Private/Editor/ControlRigEditor.cpp
--- a/Engine/Plugins/Animation/ControlRig/Source/ControlRigEditor/Private/Editor/ControlRigEditor.cpp
+++ b/Engine/Plugins/Animation/ControlRig/Source/ControlRigEditor/Private/Editor/ControlRigEditor.cpp
@@ -3629,6 +3629,13 @@ void FControlRigEditor::HandlePreviewSceneCreated(const TSharedRef<IPersonaPrevi
 			{
 				const int32 CapsuleIndex = InIndices[Index];
 				const int32 HierarchyIndex = CapsuleToHierarchyIndex.FindChecked(CapsuleIndex);
+
+				if (!Hierarchy->IsValidIndex(HierarchyIndex) || !OutCapsules.IsValidIndex(CapsuleIndex))
+				{
+					UpdateCapsules();
+					return;
+				}
+				
 				check(Hierarchy->IsValidIndex(HierarchyIndex));
 				check(OutCapsules.IsValidIndex(CapsuleIndex));
 
diff --git a/Engine/Plugins/Experimental/ControlRigSpline/Source/ControlRigSpline/Private/ControlRigSplineTypes.cpp b/Engine/Plugins/Experimental/ControlRigSpline/Source/ControlRigSpline/Private/ControlRigSplineTypes.cpp
--- a/Engine/Plugins/Experimental/ControlRigSpline/Source/ControlRigSpline/Private/ControlRigSplineTypes.cpp
+++ b/Engine/Plugins/Experimental/ControlRigSpline/Source/ControlRigSpline/Private/ControlRigSplineTypes.cpp
@@ -110,33 +110,6 @@ void FControlRigSpline::SetControlPoints(const TArrayView<const FVector>& InPoin
 					}
 				}
 
-				// Correct length of samples
-				if (!bSplineModeChanged && !bSamplesCountChanged && !bNumControlPointsChanged)
-				{
-					if (SplineData->InitialLengths.Num() > 0)
-					{
-						TArray<FVector> SamplesBeforeCorrect = SplineData->SamplesArray;
-						for (int32 i = 0; i < ControlPointsCount - 1; ++i)
-						{
-							for (int32 j = (i == 0) ? 1 : 0; j < SamplesPerSegment; ++j)
-							{
-								// Get direction from samples before correction
-								FVector Dir = SamplesBeforeCorrect[i * SamplesPerSegment + j] - SamplesBeforeCorrect[i * SamplesPerSegment + j - 1];
-								Dir.Normalize();
-
-								float InitialLength = SplineData->InitialLengths[i * SamplesPerSegment + j - 1];
-								// Current length as the projection on the Dir vector (might be negative)
-								float CurrentLength = (SplineData->SamplesArray[i * SamplesPerSegment + j] - SplineData->SamplesArray[i * SamplesPerSegment + j - 1]).Dot(Dir);
-								float FixedLength = FMath::Clamp(CurrentLength,
-									(Compression > 0.f) ? InitialLength * Compression : CurrentLength,
-									(Stretch > 0.f) ? InitialLength * Stretch : CurrentLength);
-
-								SplineData->SamplesArray[i * SamplesPerSegment + j] = SplineData->SamplesArray[i * SamplesPerSegment + j - 1] + Dir * FixedLength;
-							}
-						}
-					}
-				}
-
 				// tinySpline will allocate the samples array, but does not free that memory. We need to take care of that.
 				free(SamplesPtr);
 
@@ -160,33 +133,6 @@ void FControlRigSpline::SetControlPoints(const TArrayView<const FVector>& InPoin
 				// tinySpline will allocate the samples array, but does not free that memory. We need to take care of that.
 				free(SamplesPtr);
 
-				// Correct length of samples
-				if (!bSplineModeChanged && !bSamplesCountChanged && !bNumControlPointsChanged)
-				{
-					if (SplineData->InitialLengths.Num() > 0)
-					{
-						TArray<FVector> SamplesBeforeCorrect = SplineData->SamplesArray;
-						for (int32 i = 0; i < ControlPointsCount - 1; ++i)
-						{
-							for (int32 j = (i == 0) ? 1 : 0; j < SamplesPerSegment; ++j)
-							{
-								// Get direction from samples before correction
-								FVector Dir = SamplesBeforeCorrect[i * SamplesPerSegment + j] - SamplesBeforeCorrect[i * SamplesPerSegment + j - 1];
-								Dir.Normalize();
-
-								float InitialLength = SplineData->InitialLengths[i * SamplesPerSegment + j - 1];
-								// Current length as the projection on the Dir vector (might be negative)
-								float CurrentLength = (SplineData->SamplesArray[i * SamplesPerSegment + j] - SplineData->SamplesArray[i * SamplesPerSegment + j - 1]).Dot(Dir);
-								float FixedLength = FMath::Clamp(CurrentLength,
-                                	(Compression > 0.f) ? InitialLength * Compression : CurrentLength,
-                                	(Stretch > 0.f) ? InitialLength * Stretch : CurrentLength);
-
-								SplineData->SamplesArray[i * SamplesPerSegment + j] = SplineData->SamplesArray[i * SamplesPerSegment + j - 1] + Dir * FixedLength;
-							}
-						}
-					}
-				}
-
 				break;
 			}
 			default:
@@ -196,21 +142,50 @@ void FControlRigSpline::SetControlPoints(const TArrayView<const FVector>& InPoin
 			}
 		}
 
-		// Cache accumulated length at sample array
-		SplineData->AccumulatedLenth.SetNumUninitialized(SplineData->SamplesArray.Num(), false);
-		SplineData->AccumulatedLenth[0] = 0.f;
-		if (bSplineModeChanged || bSamplesCountChanged || bNumControlPointsChanged)
+		// Correct length of samples
+		if (!bSplineModeChanged && !bSamplesCountChanged && !bNumControlPointsChanged)
 		{
-			SplineData->InitialLengths.SetNumUninitialized(SplineData->SamplesArray.Num() - 1);
+			if (SplineData->InitialLengths.Num() > 0)
+			{
+				TArray<FVector> SamplesBeforeCorrect = SplineData->SamplesArray;
+				for (int32 i = 0; i < ControlPointsCount - 1; ++i)
+				{
+					for (int32 j = (i == 0) ? 1 : 0; j < SamplesPerSegment; ++j)
+					{
+						// Get direction from samples before correction
+						FVector Dir = SamplesBeforeCorrect[i * SamplesPerSegment + j] - SamplesBeforeCorrect[i * SamplesPerSegment + j - 1];
+						Dir.Normalize();
+
+						float InitialLength = SplineData->InitialLengths[i * SamplesPerSegment + j - 1];
+						// Current length as the projection on the Dir vector (might be negative)
+						float CurrentLength = (SplineData->SamplesArray[i * SamplesPerSegment + j] - SplineData->SamplesArray[i * SamplesPerSegment + j - 1]).Dot(Dir);
+						float FixedLength = FMath::Clamp(CurrentLength,
+							(Compression > 0.f) ? InitialLength * Compression : CurrentLength,
+							(Stretch > 0.f) ? InitialLength * Stretch : CurrentLength);
+
+						SplineData->SamplesArray[i * SamplesPerSegment + j] = SplineData->SamplesArray[i * SamplesPerSegment + j - 1] + Dir * FixedLength;
+					}
+				}
+			}
 		}
-		for (int32 i = 1; i < SplineData->SamplesArray.Num(); ++i)
+
+		// Cache accumulated length at sample array
 		{
-			float CurrentLength = FVector::Distance(SplineData->SamplesArray[i - 1], SplineData->SamplesArray[i]);
+			SplineData->AccumulatedLenth.SetNumUninitialized(SplineData->SamplesArray.Num(), false);
+			SplineData->AccumulatedLenth[0] = 0.f;
 			if (bSplineModeChanged || bSamplesCountChanged || bNumControlPointsChanged)
 			{
-				SplineData->InitialLengths[i-1] = CurrentLength;
+				SplineData->InitialLengths.SetNumUninitialized(SplineData->SamplesArray.Num() - 1);
+			}
+			for (int32 i = 1; i < SplineData->SamplesArray.Num(); ++i)
+			{
+				float CurrentLength = FVector::Distance(SplineData->SamplesArray[i - 1], SplineData->SamplesArray[i]);
+				if (bSplineModeChanged || bSamplesCountChanged || bNumControlPointsChanged)
+				{
+					SplineData->InitialLengths[i-1] = CurrentLength;
+				}
+				SplineData->AccumulatedLenth[i] = SplineData->AccumulatedLenth[i - 1] + CurrentLength;
 			}
-			SplineData->AccumulatedLenth[i] = SplineData->AccumulatedLenth[i - 1] + CurrentLength;
 		}
 	}	
 }
diff --git a/Engine/Plugins/Experimental/ControlRigSpline/Source/ControlRigSpline/Private/ControlRigSplineUnits.cpp b/Engine/Plugins/Experimental/ControlRigSpline/Source/ControlRigSpline/Private/ControlRigSplineUnits.cpp
--- a/Engine/Plugins/Experimental/ControlRigSpline/Source/ControlRigSpline/Private/ControlRigSplineUnits.cpp
+++ b/Engine/Plugins/Experimental/ControlRigSpline/Source/ControlRigSpline/Private/ControlRigSplineUnits.cpp
@@ -584,3 +584,54 @@ FRigUnit_FitChainToSplineCurve_Execute()
 		Context.DrawInterface->DrawPoints(DebugSettings.WorldOffset, CurvePositions, DebugSettings.Scale * 4.f, DebugSettings.SegmentsColor);
 	}
 }
+
+FRigUnit_FitSplineCurveToChain_Execute()
+{
+    DECLARE_SCOPE_HIERARCHICAL_COUNTER_RIGUNIT()
+
+	URigHierarchy* Hierarchy = ExecuteContext.Hierarchy;
+	if (Hierarchy == nullptr)
+	{
+		return;
+	}
+
+	if (!Spline.SplineData.IsValid())
+	{
+		return;
+	}
+
+	if (Context.State == EControlRigState::Init)
+	{
+		return;
+	}
+
+	// 1.- Create spline from chain
+	TArray<FVector> AuxControlPoints;
+	AuxControlPoints.SetNumUninitialized(Items.Num());	
+	for (int32 i = 0; i < Items.Num(); ++i)
+	{
+		const FRigElementKey& Key = Items[i];
+		AuxControlPoints[i] = Hierarchy->GetGlobalTransform(Key).GetLocation();
+	}
+    FControlRigSpline AuxSpline;
+	const TArrayView<const FVector> PointsView(AuxControlPoints.GetData(), AuxControlPoints.Num());
+	AuxSpline.SetControlPoints(PointsView, ESplineType::Hermite, Spline.SplineData->SamplesPerSegment);
+	
+
+	// 2.-  for each control point in the original spline
+	//			figure out its u in the original spline (preprocess)
+	//			query position at u on the new spline
+	const TArray<FVector>& ControlPoints = Spline.SplineData->ControlPoints;
+	TArray<FVector> NewControlPoints;
+	NewControlPoints.SetNumUninitialized(ControlPoints.Num());
+	for (int32 i = 0; i < ControlPoints.Num(); ++i)
+	{
+		float U = i / (float)(ControlPoints.Num() - 1);
+		FVector NewPosition = AuxSpline.PositionAtParam(U);
+		NewControlPoints[i] = NewPosition;
+	}
+
+	const TArrayView<const FVector> NewPointsView(NewControlPoints.GetData(), NewControlPoints.Num());
+	Spline.SetControlPoints(NewPointsView, Spline.SplineData->SplineMode, Spline.SplineData->SamplesPerSegment, Spline.SplineData->Compression, Spline.SplineData->Stretch);
+}
+
diff --git a/Engine/Plugins/Experimental/ControlRigSpline/Source/ControlRigSpline/Public/ControlRigSplineUnits.h b/Engine/Plugins/Experimental/ControlRigSpline/Source/ControlRigSpline/Public/ControlRigSplineUnits.h
--- a/Engine/Plugins/Experimental/ControlRigSpline/Source/ControlRigSpline/Public/ControlRigSplineUnits.h
+++ b/Engine/Plugins/Experimental/ControlRigSpline/Source/ControlRigSpline/Public/ControlRigSplineUnits.h
@@ -293,4 +293,35 @@ struct CONTROLRIGSPLINE_API FRigUnit_FitChainToSplineCurve : public FRigUnit_Hig
 
 	UPROPERTY(transient)
 	FRigUnit_FitChainToCurve_WorkData WorkData;
+};
+
+/**
+ * Fits a given spline curve to a chain.
+ */
+USTRUCT(meta=(DisplayName="Fit Spline Curve on Chain", Category="Hierarchy", Keywords="Fit,Resample,Spline"))
+struct CONTROLRIGSPLINE_API FRigUnit_FitSplineCurveToChain : public FRigUnit_HighlevelBaseMutable
+{
+	GENERATED_BODY()
+
+	FRigUnit_FitSplineCurveToChain()
+	{
+		
+	}
+
+	RIGVM_METHOD()
+	virtual void Execute(const FRigUnitContext& Context) override;
+
+	/** 
+	 * The items to align to
+	 */
+	UPROPERTY(meta = (Input))
+	FRigElementKeyCollection Items;
+
+	/** 
+	 * The curve to align
+	 */
+	UPROPERTY(meta = (Input, Output))
+	FControlRigSpline Spline;
+
+	
 };
\ No newline at end of file
diff --git a/Engine/Source/Runtime/RigVM/Private/RigVMCore/RigVM.cpp b/Engine/Source/Runtime/RigVM/Private/RigVMCore/RigVM.cpp
--- a/Engine/Source/Runtime/RigVM/Private/RigVMCore/RigVM.cpp
+++ b/Engine/Source/Runtime/RigVM/Private/RigVMCore/RigVM.cpp
@@ -1344,7 +1344,7 @@ bool URigVM::ShouldHaltAtInstruction(const uint16 InstructionIndex)
 			{
 				if (DifferenceIndex == INDEX_NONE)
 				{
-					if (CurrentCallstack->Last() != PreviousCallstack.Last())
+					if (!CurrentCallstack->IsEmpty() && !PreviousCallstack.IsEmpty() && CurrentCallstack->Last() != PreviousCallstack.Last())
 					{
 						NewBreakpointNode = CurrentCallstack->operator[](FMath::Min(PreviousCallstack.Num(), CurrentCallstack->Num()-1));
 					}
