Control Rig Spline: Add GetTangent node

diff --git a/Engine/Plugins/Experimental/ControlRigSpline/Source/ControlRigSpline/Private/ControlRigSplineUnits.cpp b/Engine/Plugins/Experimental/ControlRigSpline/Source/ControlRigSpline/Private/ControlRigSplineUnits.cpp
--- a/Engine/Plugins/Experimental/ControlRigSpline/Source/ControlRigSpline/Private/ControlRigSplineUnits.cpp
+++ b/Engine/Plugins/Experimental/ControlRigSpline/Source/ControlRigSpline/Private/ControlRigSplineUnits.cpp
@@ -121,6 +121,42 @@ FRigUnit_TransformFromControlRigSpline_Execute()
 	}
 }
 
+FRigUnit_TangentFromControlRigSpline_Execute()
+{
+	if (!Spline.SplineData.IsValid())
+	{
+		return;
+	}
+	
+	switch (Context.State)
+	{
+		case EControlRigState::Init:
+		case EControlRigState::Update:
+		{
+			const float ClampedU = FMath::Clamp<float>(U, 0.f, 1.f);
+
+			FVector NewTangent = Spline.TangentAtParam(ClampedU);
+
+			// Check if Tangent can be normalized. If not, keep the same tangent as before.
+			if (!NewTangent.Normalize())
+			{
+				NewTangent = Tangent;
+			}
+			else
+			{
+				Tangent = NewTangent;
+			}
+			
+			break;
+		}
+		default:
+		{
+			checkNoEntry(); // Execute is only defined for Init and Update
+			break;
+		}
+	}
+}
+
 FRigUnit_DrawControlRigSpline_Execute()
 {
 	if (Context.State == EControlRigState::Init)
diff --git a/Engine/Plugins/Experimental/ControlRigSpline/Source/ControlRigSpline/Public/ControlRigSplineUnits.h b/Engine/Plugins/Experimental/ControlRigSpline/Source/ControlRigSpline/Public/ControlRigSplineUnits.h
--- a/Engine/Plugins/Experimental/ControlRigSpline/Source/ControlRigSpline/Public/ControlRigSplineUnits.h
+++ b/Engine/Plugins/Experimental/ControlRigSpline/Source/ControlRigSpline/Public/ControlRigSplineUnits.h
@@ -21,10 +21,10 @@ struct CONTROLRIGSPLINE_API FRigUnit_ControlRigSplineFromPoints : public FRigUni
 
 	FRigUnit_ControlRigSplineFromPoints()
 	{
-		SplineMode = ESplineType::BSpline;
+		SplineMode = ESplineType::Hermite;
 		SamplesPerSegment = 16;
-		Compression = 1.f;
-		Stretch = 1.f;
+		Compression = 0.f;
+		Stretch = 0.f;
 	}
 
 	/** Execute logic for this rig unit */
@@ -129,6 +129,31 @@ struct CONTROLRIGSPLINE_API FRigUnit_TransformFromControlRigSpline : public FRig
 	FTransform Transform;
 };
 
+USTRUCT(meta = (DisplayName = "Tangent From Spline", Category = "Control Rig"))
+struct CONTROLRIGSPLINE_API FRigUnit_TangentFromControlRigSpline : public FRigUnit_ControlRigSplineBase
+{
+	GENERATED_BODY()
+
+	FRigUnit_TangentFromControlRigSpline()
+	{
+		U = 0.f;
+		Tangent = FVector::ZeroVector;
+	}
+
+	/** Execute logic for this rig unit */
+	RIGVM_METHOD()
+	virtual void Execute(const FRigUnitContext& Context) override;
+
+	UPROPERTY(meta = (Input))
+	FControlRigSpline Spline;
+
+	UPROPERTY(meta = (Input))
+	float U;
+
+	UPROPERTY(meta = (Output))
+	FVector Tangent;
+};
+
 USTRUCT(meta = (DisplayName = "Draw Spline", Category = "Control Rig"))
 struct CONTROLRIGSPLINE_API FRigUnit_DrawControlRigSpline : public FRigUnitMutable
 {
