Control Rig Spline: Improvements to spline data structure to support arrays

* implement copy operators to support copy for the tsharedptr
* make sure to clear the spline in the creator nodes so that it works for arrays

diff --git a/Engine/Plugins/Experimental/ControlRigSpline/Source/ControlRigSpline/Private/ControlRigSplineTypes.cpp b/Engine/Plugins/Experimental/ControlRigSpline/Source/ControlRigSpline/Private/ControlRigSplineTypes.cpp
--- a/Engine/Plugins/Experimental/ControlRigSpline/Source/ControlRigSpline/Private/ControlRigSplineTypes.cpp
+++ b/Engine/Plugins/Experimental/ControlRigSpline/Source/ControlRigSpline/Private/ControlRigSplineTypes.cpp
@@ -2,6 +2,17 @@
 
 #include "ControlRigSplineTypes.h"
 
+FControlRigSpline::FControlRigSpline(const FControlRigSpline& InOther)
+{
+	SplineData = InOther.SplineData;
+}
+
+FControlRigSpline& FControlRigSpline::operator=(const FControlRigSpline& InOther)
+{
+	SplineData = InOther.SplineData;
+	return *this;
+}
+
 void FControlRigSpline::SetControlPoints(const TArrayView<const FVector>& InPoints, const ESplineType SplineMode, const int32 SamplesPerSegment, const float Compression, const float Stretch)
 {
 	const int32 ControlPointsCount = InPoints.Num();
diff --git a/Engine/Plugins/Experimental/ControlRigSpline/Source/ControlRigSpline/Private/ControlRigSplineUnits.cpp b/Engine/Plugins/Experimental/ControlRigSpline/Source/ControlRigSpline/Private/ControlRigSplineUnits.cpp
--- a/Engine/Plugins/Experimental/ControlRigSpline/Source/ControlRigSpline/Private/ControlRigSplineUnits.cpp
+++ b/Engine/Plugins/Experimental/ControlRigSpline/Source/ControlRigSpline/Private/ControlRigSplineUnits.cpp
@@ -13,6 +13,9 @@ FRigUnit_ControlRigSplineFromPoints_Execute()
 		case EControlRigState::Init:
 		case EControlRigState::Update:
 		{
+			// reset the spline
+			Spline = FControlRigSpline();
+				
 			const TArrayView<const FVector> PointsView(Points.GetData(), Points.Num());
 			Spline.SetControlPoints(PointsView, SplineMode, SamplesPerSegment, Compression, Stretch);
 			break;
diff --git a/Engine/Plugins/Experimental/ControlRigSpline/Source/ControlRigSpline/Public/ControlRigSplineTypes.h b/Engine/Plugins/Experimental/ControlRigSpline/Source/ControlRigSpline/Public/ControlRigSplineTypes.h
--- a/Engine/Plugins/Experimental/ControlRigSpline/Source/ControlRigSpline/Public/ControlRigSplineTypes.h
+++ b/Engine/Plugins/Experimental/ControlRigSpline/Source/ControlRigSpline/Public/ControlRigSplineTypes.h
@@ -72,6 +72,9 @@ struct CONTROLRIGSPLINE_API FControlRigSpline
 
 	virtual ~FControlRigSpline() {}
 
+	FControlRigSpline(const FControlRigSpline& InOther);
+	FControlRigSpline& operator =(const FControlRigSpline& InOther);
+
 	TSharedPtr<FControlRigSplineImpl> SplineData;
 
 	/**
diff --git a/Engine/Plugins/Experimental/ControlRigSpline/Source/ControlRigSpline/Public/ControlRigSplineUnits.h b/Engine/Plugins/Experimental/ControlRigSpline/Source/ControlRigSpline/Public/ControlRigSplineUnits.h
--- a/Engine/Plugins/Experimental/ControlRigSpline/Source/ControlRigSpline/Public/ControlRigSplineUnits.h
+++ b/Engine/Plugins/Experimental/ControlRigSpline/Source/ControlRigSpline/Public/ControlRigSplineUnits.h
@@ -14,7 +14,7 @@ struct CONTROLRIGSPLINE_API FRigUnit_ControlRigSplineBase : public FRigUnit
 	GENERATED_BODY()
 };
 
-USTRUCT(meta = (DisplayName = "Spline From Points", Keywords="Spline From Positions", Category = "Control Rig"))
+USTRUCT(meta = (DisplayName = "Spline From Points", Keywords="Spline From Positions", Category = "Control Rig", Varying))
 struct CONTROLRIGSPLINE_API FRigUnit_ControlRigSplineFromPoints : public FRigUnit_ControlRigSplineBase
 {
 	GENERATED_BODY()
@@ -50,7 +50,7 @@ struct CONTROLRIGSPLINE_API FRigUnit_ControlRigSplineFromPoints : public FRigUni
 	FControlRigSpline Spline;
 };
 
-USTRUCT(meta = (DisplayName = "Set Spline Points", Category = "Control Rig"))
+USTRUCT(meta = (DisplayName = "Set Spline Points", Category = "Control Rig", Varying))
 struct CONTROLRIGSPLINE_API FRigUnit_SetSplinePoints : public FRigUnitMutable
 {
 	GENERATED_BODY()
