Control Rig Spline: Adding node Closest Parameter to Point

diff --git a/Engine/Plugins/Experimental/ControlRigSpline/Source/ControlRigSpline/Private/ControlRigSplineUnits.cpp b/Engine/Plugins/Experimental/ControlRigSpline/Source/ControlRigSpline/Private/ControlRigSplineUnits.cpp
--- a/Engine/Plugins/Experimental/ControlRigSpline/Source/ControlRigSpline/Private/ControlRigSplineUnits.cpp
+++ b/Engine/Plugins/Experimental/ControlRigSpline/Source/ControlRigSpline/Private/ControlRigSplineUnits.cpp
@@ -674,3 +674,66 @@ FRigUnit_FitSplineCurveToChain_Execute()
 	Spline.SetControlPoints(NewPointsView, Spline.SplineData->SplineMode, Spline.SplineData->SamplesPerSegment, Spline.SplineData->Compression, Spline.SplineData->Stretch);
 }
 
+FRigUnit_ClosestParameterFromControlRigSpline_Execute()
+{
+	if (!Spline.SplineData.IsValid())
+	{
+		return;
+	}
+	
+	switch (Context.State)
+	{
+		case EControlRigState::Init:
+		case EControlRigState::Update:
+		{
+				const int32 SampleCount = Spline.SplineData->SamplesArray.Num();
+
+				auto DistanceToSegment = [&](int32 Index0, int32 Index1, float& OutParam) -> float
+				{
+					const int32 MaxIndex = Spline.SplineData->SamplesArray.Num()-1;
+					const FVector& P0 = Spline.SplineData->SamplesArray[Index0];
+					const FVector& P1 = Spline.SplineData->SamplesArray[Index1];
+
+					const FVector V = P1 - P0;
+					const FVector U = P0 - Position;
+					const float T = - U.Dot(V) / V.Dot(V);
+					if (T > 0 && T < 1)
+					{
+						OutParam = (1-T)*(Index0/(float)MaxIndex) + T*(Index1/(float)MaxIndex);
+						return U.Cross(V).Length() / V.Length();
+					}
+
+					const float D0 = FVector::Distance(P0, Position);
+					const float D1 = FVector::Distance(P1, Position);
+					if (D0 < D1)
+					{
+						OutParam = Index0/(float)MaxIndex;
+						return D0;
+					}
+
+					OutParam = Index1/(float)MaxIndex;
+					return D1;
+				};
+
+				float ClosestDistance = TNumericLimits<float>::Max();
+				for (int32 i=1; i<SampleCount; ++i)
+				{
+					float Param;
+					float Distance = DistanceToSegment(i-1, i, Param);
+
+					if (Distance < ClosestDistance)
+					{
+						ClosestDistance = Distance;
+						U = Param;
+					}
+				}				
+				
+			break;
+		}
+		default:
+		{
+			checkNoEntry(); // Execute is only defined for Init and Update
+			break;
+		}
+	}
+}
diff --git a/Engine/Plugins/Experimental/ControlRigSpline/Source/ControlRigSpline/Public/ControlRigSplineUnits.h b/Engine/Plugins/Experimental/ControlRigSpline/Source/ControlRigSpline/Public/ControlRigSplineUnits.h
--- a/Engine/Plugins/Experimental/ControlRigSpline/Source/ControlRigSpline/Public/ControlRigSplineUnits.h
+++ b/Engine/Plugins/Experimental/ControlRigSpline/Source/ControlRigSpline/Public/ControlRigSplineUnits.h
@@ -346,7 +346,32 @@ struct CONTROLRIGSPLINE_API FRigUnit_FitSplineCurveToChain : public FRigUnit_Hig
 	 * The curve to align
 	 */
 	UPROPERTY(meta = (Input, Output))
+	FControlRigSpline Spline;	
+};
+
+USTRUCT(meta = (DisplayName = "Closest Parameter From Spline", Category = "Control Rig"))
+struct CONTROLRIGSPLINE_API FRigUnit_ClosestParameterFromControlRigSpline : public FRigUnit_ControlRigSplineBase
+{
+	GENERATED_BODY()
+
+	FRigUnit_ClosestParameterFromControlRigSpline()
+	{
+		Position = FVector::ZeroVector;
+		U = 0.f;
+	}
+
+	/** Execute logic for this rig unit */
+	RIGVM_METHOD()
+	virtual void Execute(const FRigUnitContext& Context) override;
+
+	UPROPERTY(meta = (Input))
 	FControlRigSpline Spline;
 
+	UPROPERTY(meta = (Input))
+	FVector Position;
+
+	UPROPERTY(meta = (Output))
+	float U;
 	
-};
\ No newline at end of file
+};
+
