Control Rig Spine: Report error when spline cannot be created

diff --git a/Engine/Plugins/Experimental/ControlRigSpline/Source/ControlRigSpline/Private/ControlRigSplineUnits.cpp b/Engine/Plugins/Experimental/ControlRigSpline/Source/ControlRigSpline/Private/ControlRigSplineUnits.cpp
--- a/Engine/Plugins/Experimental/ControlRigSpline/Source/ControlRigSpline/Private/ControlRigSplineUnits.cpp
+++ b/Engine/Plugins/Experimental/ControlRigSpline/Source/ControlRigSpline/Private/ControlRigSplineUnits.cpp
@@ -15,6 +15,12 @@ FRigUnit_ControlRigSplineFromPoints_Execute()
 		{
 			// reset the spline
 			Spline = FControlRigSpline();
+
+			if (Points.Num() < 4)
+			{
+				UE_CONTROLRIG_RIGUNIT_REPORT_WARNING(TEXT("Cannot create a spline with less than 4 points (%d received)."), Points.Num());
+				return;
+			}
 				
 			const TArrayView<const FVector> PointsView(Points.GetData(), Points.Num());
 			Spline.SetControlPoints(PointsView, SplineMode, SamplesPerSegment, Compression, Stretch);
@@ -644,6 +650,12 @@ FRigUnit_FitSplineCurveToChain_Execute()
 		return;
 	}
 
+	if (Items.Num() < 4)
+	{
+		UE_CONTROLRIG_RIGUNIT_REPORT_WARNING(TEXT("Cannot create a spline with less than 4 points (%d received)."), Items.Num());
+		return;
+	}
+
 	// 1.- Create spline from chain
 	TArray<FVector> AuxControlPoints;
 	AuxControlPoints.SetNumUninitialized(Items.Num());	
