Control Rig Spline: Adapt to LWC

diff --git a/Engine/Plugins/Experimental/ControlRigSpline/Source/ControlRigSpline/ControlRigSpline.build.cs b/Engine/Plugins/Experimental/ControlRigSpline/Source/ControlRigSpline/ControlRigSpline.build.cs
--- a/Engine/Plugins/Experimental/ControlRigSpline/Source/ControlRigSpline/ControlRigSpline.build.cs
+++ b/Engine/Plugins/Experimental/ControlRigSpline/Source/ControlRigSpline/ControlRigSpline.build.cs
@@ -7,8 +7,12 @@ namespace UnrealBuildTool.Rules
 		public ControlRigSpline(ReadOnlyTargetRules Target) : base(Target)
 		{
 			bEnableUndefinedIdentifierWarnings = false;
+			
+			if (Target.bLWCDisabled)
+			{
+				PublicDefinitions.Add("TINYSPLINE_FLOAT_PRECISION=1");
+			}
 
-			PublicDefinitions.Add("TINYSPLINE_FLOAT_PRECISION=1");
 			PrivateIncludePaths.Add("ControlRigSpline/ThirdParty/TinySpline");
 
 			PublicDependencyModuleNames.AddRange(
diff --git a/Engine/Plugins/Experimental/ControlRigSpline/Source/ControlRigSpline/Private/ControlRigSplineTypes.cpp b/Engine/Plugins/Experimental/ControlRigSpline/Source/ControlRigSpline/Private/ControlRigSplineTypes.cpp
--- a/Engine/Plugins/Experimental/ControlRigSpline/Source/ControlRigSpline/Private/ControlRigSplineTypes.cpp
+++ b/Engine/Plugins/Experimental/ControlRigSpline/Source/ControlRigSpline/Private/ControlRigSplineTypes.cpp
@@ -61,7 +61,7 @@ void FControlRigSpline::SetControlPoints(const TArrayView<const FVector>& InPoin
 				// so we need to iterate updating the points one by one.
 				for (int32 i = 0; i < ControlPointsCount; ++i)
 				{
-					FVector3f Point = InPoints[i];
+					FVector Point = InPoints[i];
 					ts_bspline_set_control_point_at(SplineData->Spline.data(), i, &Point.X, nullptr);
 				}
 
@@ -75,16 +75,16 @@ void FControlRigSpline::SetControlPoints(const TArrayView<const FVector>& InPoin
 						
 				// Create the control points array for the interpolation.
 				// The first dimension acts as the query parameter, the rest of the dimensions are the XYZ of the control points
-				TArray<float> ControlPointsArray;
+				TArray<FVector::FReal> ControlPointsArray;
 				ControlPointsArray.SetNumUninitialized(ControlPointsCount * 4, false);
 	
 				for (int32 i = 0; i < ControlPointsCount; ++i)
 				{
-					ControlPointsArray[i * 4 + 0] = i / (float)(ControlPointsCount - 1);
+					ControlPointsArray[i * 4 + 0] = i / (FVector::FReal)(ControlPointsCount - 1);
 
-					ControlPointsArray[i * 4 + 1] = (float)InPoints[i][0];
-					ControlPointsArray[i * 4 + 2] = (float)InPoints[i][1];
-					ControlPointsArray[i * 4 + 3] = (float)InPoints[i][2];
+					ControlPointsArray[i * 4 + 1] = (FVector::FReal)InPoints[i][0];
+					ControlPointsArray[i * 4 + 2] = (FVector::FReal)InPoints[i][1];
+					ControlPointsArray[i * 4 + 3] = (FVector::FReal)InPoints[i][2];
 				}
 
 				// Create a new spline to store the result of the interpolation
@@ -109,7 +109,7 @@ void FControlRigSpline::SetControlPoints(const TArrayView<const FVector>& InPoin
 			case ESplineType::BSpline:
 			{
 				// Cache sample positions of the spline
-				float* SamplesPtr = nullptr;
+				FVector::FReal* SamplesPtr = nullptr;
 				size_t ActualSamplesPerSegment = 0;
 				ts_bspline_sample(SplineData->Spline.data(), (ControlPointsCount - 1) * SamplesPerSegment, &SamplesPtr, &ActualSamplesPerSegment, nullptr);
 				SplineData->SamplesArray.SetNumUninitialized((ControlPointsCount - 1) * SamplesPerSegment, false);
@@ -117,7 +117,7 @@ void FControlRigSpline::SetControlPoints(const TArrayView<const FVector>& InPoin
 				{
 					for (int32 j = 0; j < SamplesPerSegment; ++j)
 					{
-						memcpy(&SplineData->SamplesArray[i * SamplesPerSegment + j].X, SamplesPtr + (i * SamplesPerSegment + j) * 3, sizeof(float) * 3);
+						memcpy(&SplineData->SamplesArray[i * SamplesPerSegment + j].X, SamplesPtr + (i * SamplesPerSegment + j) * 3, sizeof(FVector::FReal) * 3);
 					}
 				}
 
@@ -129,7 +129,7 @@ void FControlRigSpline::SetControlPoints(const TArrayView<const FVector>& InPoin
 			case ESplineType::Hermite:
 			{
 				// Cache sample positions of the spline
-				float* SamplesPtr = nullptr;
+				FVector::FReal* SamplesPtr = nullptr;
 				size_t ActualSamplesPerSegment = 0;
 				ts_bspline_sample(SplineData->Spline.data(), (ControlPointsCount - 1) * SamplesPerSegment, &SamplesPtr, &ActualSamplesPerSegment, nullptr);
 				SplineData->SamplesArray.SetNumUninitialized((ControlPointsCount - 1) * SamplesPerSegment);
@@ -137,7 +137,7 @@ void FControlRigSpline::SetControlPoints(const TArrayView<const FVector>& InPoin
 				{
 					for (int32 j = 0; j < SamplesPerSegment; ++j)
 					{
-						memcpy(&SplineData->SamplesArray[i * SamplesPerSegment + j].X, SamplesPtr + (i * SamplesPerSegment + j) * 4 + 1, sizeof(float) * 3);
+						memcpy(&SplineData->SamplesArray[i * SamplesPerSegment + j].X, SamplesPtr + (i * SamplesPerSegment + j) * 4 + 1, sizeof(FVector::FReal) * 3);
 					}
 				}
 
@@ -206,7 +206,7 @@ int32 FControlRigSpline::GetControlPoints(TArray<FVector>& OutPoints) const
 	const size_t Count = SplineData->Spline.numControlPoints();
 	OutPoints.SetNumUninitialized(Count);
 
-	float* Points;
+	FVector::FReal* Points;
 	ts_bspline_control_points(SplineData->Spline.data(), &Points, nullptr);
 
 	switch (SplineData->SplineMode)
diff --git a/Engine/Plugins/Experimental/ControlRigSpline/Source/ControlRigSplineEditor/ControlRigSplineEditor.build.cs b/Engine/Plugins/Experimental/ControlRigSpline/Source/ControlRigSplineEditor/ControlRigSplineEditor.build.cs
--- a/Engine/Plugins/Experimental/ControlRigSpline/Source/ControlRigSplineEditor/ControlRigSplineEditor.build.cs
+++ b/Engine/Plugins/Experimental/ControlRigSpline/Source/ControlRigSplineEditor/ControlRigSplineEditor.build.cs
@@ -8,7 +8,10 @@ namespace UnrealBuildTool.Rules
 		{
 			bEnableUndefinedIdentifierWarnings = false;
 
-			PublicDefinitions.Add("TINYSPLINE_FLOAT_PRECISION=1");
+			if (Target.bLWCDisabled)
+			{
+				PublicDefinitions.Add("TINYSPLINE_FLOAT_PRECISION=1");
+			}
 			PrivateIncludePaths.Add("ControlRigSpline/ThirdParty/TinySpline");
 
 			PublicDependencyModuleNames.AddRange(
diff --git a/Engine/Source/Programs/UnrealBuildTool/Configuration/TargetRules.cs b/Engine/Source/Programs/UnrealBuildTool/Configuration/TargetRules.cs
--- a/Engine/Source/Programs/UnrealBuildTool/Configuration/TargetRules.cs
+++ b/Engine/Source/Programs/UnrealBuildTool/Configuration/TargetRules.cs
@@ -2155,7 +2155,12 @@ namespace UnrealBuildTool
 		{
 			get { return Inner.bCompileISPC && GlobalDefinitions.Contains("UE_LARGE_WORLD_COORDINATES_DISABLED=1"); }	// LWC_TODO: Temporarily disable ISPC when LWC is turned on. To be removed when double support is added to ISPC.
 		}
-
+		
+		public bool bLWCDisabled 
+		{
+			get { return GlobalDefinitions.Contains("UE_LARGE_WORLD_COORDINATES_DISABLED=1"); }	
+		}
+		
 		public bool bCompilePython
 		{
 			get { return Inner.bCompilePython; }
