Control Rig Spline: Fix compile error for LWC disabled in LargeWorldCoordinates.h

diff --git a/Engine/Plugins/Animation/ControlRigSpline/Source/ControlRigSpline/Private/ControlRigSplineTypes.cpp b/Engine/Plugins/Animation/ControlRigSpline/Source/ControlRigSpline/Private/ControlRigSplineTypes.cpp
--- a/Engine/Plugins/Animation/ControlRigSpline/Source/ControlRigSpline/Private/ControlRigSplineTypes.cpp
+++ b/Engine/Plugins/Animation/ControlRigSpline/Source/ControlRigSpline/Private/ControlRigSplineTypes.cpp
@@ -61,8 +61,13 @@ void FControlRigSpline::SetControlPoints(const TArrayView<const FVector>& InPoin
 				// so we need to iterate updating the points one by one.
 				for (int32 i = 0; i < ControlPointsCount; ++i)
 				{
+#if (UE_LARGE_WORLD_COORDINATES_DISABLED && !TINYSPLINE_FLOAT_PRECISION)
+					FVector3d Point = InPoints[i];
+					ts_bspline_set_control_point_at(SplineData->Spline.data(), i, &Point.X, nullptr);
+#else
 					FVector Point = InPoints[i];
 					ts_bspline_set_control_point_at(SplineData->Spline.data(), i, &Point.X, nullptr);
+#endif
 				}
 
 				break;
@@ -75,7 +80,11 @@ void FControlRigSpline::SetControlPoints(const TArrayView<const FVector>& InPoin
 						
 				// Create the control points array for the interpolation.
 				// The first dimension acts as the query parameter, the rest of the dimensions are the XYZ of the control points
+#if (UE_LARGE_WORLD_COORDINATES_DISABLED && !TINYSPLINE_FLOAT_PRECISION)					
+				TArray<double> ControlPointsArray;
+#else
 				TArray<FVector::FReal> ControlPointsArray;
+#endif
 				ControlPointsArray.SetNumUninitialized(ControlPointsCount * 4, false);
 	
 				for (int32 i = 0; i < ControlPointsCount; ++i)
@@ -109,7 +118,11 @@ void FControlRigSpline::SetControlPoints(const TArrayView<const FVector>& InPoin
 			case ESplineType::BSpline:
 			{
 				// Cache sample positions of the spline
+#if (UE_LARGE_WORLD_COORDINATES_DISABLED && !TINYSPLINE_FLOAT_PRECISION)					
+				double* SamplesPtr = nullptr;
+#else
 				FVector::FReal* SamplesPtr = nullptr;
+#endif
 				size_t ActualSamplesPerSegment = 0;
 				ts_bspline_sample(SplineData->Spline.data(), (ControlPointsCount - 1) * SamplesPerSegment, &SamplesPtr, &ActualSamplesPerSegment, nullptr);
 				SplineData->SamplesArray.SetNumUninitialized((ControlPointsCount - 1) * SamplesPerSegment, false);
@@ -117,7 +130,9 @@ void FControlRigSpline::SetControlPoints(const TArrayView<const FVector>& InPoin
 				{
 					for (int32 j = 0; j < SamplesPerSegment; ++j)
 					{
-						memcpy(&SplineData->SamplesArray[i * SamplesPerSegment + j].X, SamplesPtr + (i * SamplesPerSegment + j) * 3, sizeof(FVector::FReal) * 3);
+						SplineData->SamplesArray[i * SamplesPerSegment + j].X = SamplesPtr[(i * SamplesPerSegment + j) * 3];
+						SplineData->SamplesArray[i * SamplesPerSegment + j].Y = SamplesPtr[(i * SamplesPerSegment + j) * 3 + 1];
+						SplineData->SamplesArray[i * SamplesPerSegment + j].Z = SamplesPtr[(i * SamplesPerSegment + j) * 3 + 2];
 					}
 				}
 
@@ -129,7 +144,11 @@ void FControlRigSpline::SetControlPoints(const TArrayView<const FVector>& InPoin
 			case ESplineType::Hermite:
 			{
 				// Cache sample positions of the spline
+#if (UE_LARGE_WORLD_COORDINATES_DISABLED && !TINYSPLINE_FLOAT_PRECISION)					
+				double* SamplesPtr = nullptr;
+#else
 				FVector::FReal* SamplesPtr = nullptr;
+#endif
 				size_t ActualSamplesPerSegment = 0;
 				ts_bspline_sample(SplineData->Spline.data(), (ControlPointsCount - 1) * SamplesPerSegment, &SamplesPtr, &ActualSamplesPerSegment, nullptr);
 				SplineData->SamplesArray.SetNumUninitialized((ControlPointsCount - 1) * SamplesPerSegment);
@@ -137,7 +156,9 @@ void FControlRigSpline::SetControlPoints(const TArrayView<const FVector>& InPoin
 				{
 					for (int32 j = 0; j < SamplesPerSegment; ++j)
 					{
-						memcpy(&SplineData->SamplesArray[i * SamplesPerSegment + j].X, SamplesPtr + (i * SamplesPerSegment + j) * 4 + 1, sizeof(FVector::FReal) * 3);
+						SplineData->SamplesArray[i * SamplesPerSegment + j].X = SamplesPtr[(i * SamplesPerSegment + j) * 4 + 1];
+						SplineData->SamplesArray[i * SamplesPerSegment + j].Y = SamplesPtr[(i * SamplesPerSegment + j) * 4 + 2];
+						SplineData->SamplesArray[i * SamplesPerSegment + j].Z = SamplesPtr[(i * SamplesPerSegment + j) * 4 + 3];
 					}
 				}
 
@@ -206,7 +227,12 @@ int32 FControlRigSpline::GetControlPoints(TArray<FVector>& OutPoints) const
 	const size_t Count = SplineData->Spline.numControlPoints();
 	OutPoints.SetNumUninitialized(Count);
 
+#if (UE_LARGE_WORLD_COORDINATES_DISABLED && !TINYSPLINE_FLOAT_PRECISION)					
+	double* Points;
+#else
 	FVector::FReal* Points;
+#endif	
+	
 	ts_bspline_control_points(SplineData->Spline.data(), &Points, nullptr);
 
 	switch (SplineData->SplineMode)
