Control Rig Spline: Fixed duplicated IMPLEMENT_MODULE for ControlRigSpline

diff --git a/Engine/Plugins/Animation/ControlRigSpline/Source/ControlRigSplineEditor/Private/ControlRigSplineEditorModule.cpp b/Engine/Plugins/Animation/ControlRigSpline/Source/ControlRigSplineEditor/Private/ControlRigSplineEditorModule.cpp
--- a/Engine/Plugins/Animation/ControlRigSpline/Source/ControlRigSplineEditor/Private/ControlRigSplineEditorModule.cpp
+++ b/Engine/Plugins/Animation/ControlRigSpline/Source/ControlRigSplineEditor/Private/ControlRigSplineEditorModule.cpp
@@ -12,4 +12,4 @@ void FControlRigSplineEditorModule::ShutdownModule()
 {
 }
 
-IMPLEMENT_MODULE(FControlRigSplineEditorModule, ControlRigSpline)
+IMPLEMENT_MODULE(FControlRigSplineEditorModule, ControlRigSplineEditor)
