ControlRigSpline change category

diff --git a/Engine/Plugins/Animation/ControlRigSpline/Source/ControlRigSpline/Public/ControlRigSplineUnits.h b/Engine/Plugins/Animation/ControlRigSpline/Source/ControlRigSpline/Public/ControlRigSplineUnits.h
--- a/Engine/Plugins/Animation/ControlRigSpline/Source/ControlRigSpline/Public/ControlRigSplineUnits.h
+++ b/Engine/Plugins/Animation/ControlRigSpline/Source/ControlRigSpline/Public/ControlRigSplineUnits.h
@@ -14,7 +14,7 @@ struct CONTROLRIGSPLINE_API FRigUnit_ControlRigSplineBase : public FRigUnit
 	GENERATED_BODY()
 };
 
-USTRUCT(meta = (DisplayName = "Spline From Points", Keywords="Spline From Positions", Category = "Control Rig", Varying))
+USTRUCT(meta = (DisplayName = "Spline From Points", Keywords="Spline From Positions", Category = "Spline", Varying))
 struct CONTROLRIGSPLINE_API FRigUnit_ControlRigSplineFromPoints : public FRigUnit_ControlRigSplineBase
 {
 	GENERATED_BODY()
@@ -50,7 +50,7 @@ struct CONTROLRIGSPLINE_API FRigUnit_ControlRigSplineFromPoints : public FRigUni
 	FControlRigSpline Spline;
 };
 
-USTRUCT(meta = (DisplayName = "Set Spline Points", Category = "Control Rig", Varying))
+USTRUCT(meta = (DisplayName = "Set Spline Points", Category = "Spline", Varying))
 struct CONTROLRIGSPLINE_API FRigUnit_SetSplinePoints : public FRigUnitMutable
 {
 	GENERATED_BODY()
@@ -71,7 +71,7 @@ struct CONTROLRIGSPLINE_API FRigUnit_SetSplinePoints : public FRigUnitMutable
 	FControlRigSpline Spline;
 };
 
-USTRUCT(meta = (DisplayName = "Position From Spline", Keywords="Point From Spline", Category = "Control Rig"))
+USTRUCT(meta = (DisplayName = "Position From Spline", Keywords="Point From Spline", Category = "Spline"))
 struct CONTROLRIGSPLINE_API FRigUnit_PositionFromControlRigSpline : public FRigUnit_ControlRigSplineBase
 {
 	GENERATED_BODY()
@@ -96,7 +96,7 @@ struct CONTROLRIGSPLINE_API FRigUnit_PositionFromControlRigSpline : public FRigU
 	FVector Position;
 };
 
-USTRUCT(meta = (DisplayName = "Transform From Spline", Category = "Control Rig"))
+USTRUCT(meta = (DisplayName = "Transform From Spline", Category = "Spline"))
 struct CONTROLRIGSPLINE_API FRigUnit_TransformFromControlRigSpline : public FRigUnit_ControlRigSplineBase
 {
 	GENERATED_BODY()
@@ -129,7 +129,7 @@ struct CONTROLRIGSPLINE_API FRigUnit_TransformFromControlRigSpline : public FRig
 	FTransform Transform;
 };
 
-USTRUCT(meta = (DisplayName = "Tangent From Spline", Category = "Control Rig"))
+USTRUCT(meta = (DisplayName = "Tangent From Spline", Category = "Spline"))
 struct CONTROLRIGSPLINE_API FRigUnit_TangentFromControlRigSpline : public FRigUnit_ControlRigSplineBase
 {
 	GENERATED_BODY()
@@ -154,7 +154,7 @@ struct CONTROLRIGSPLINE_API FRigUnit_TangentFromControlRigSpline : public FRigUn
 	FVector Tangent;
 };
 
-USTRUCT(meta = (DisplayName = "Draw Spline", Category = "Control Rig"))
+USTRUCT(meta = (DisplayName = "Draw Spline", Category = "Spline"))
 struct CONTROLRIGSPLINE_API FRigUnit_DrawControlRigSpline : public FRigUnitMutable
 {
 	GENERATED_BODY()
@@ -183,7 +183,7 @@ struct CONTROLRIGSPLINE_API FRigUnit_DrawControlRigSpline : public FRigUnitMutab
 	int32 Detail;
 };
 
-USTRUCT(meta = (DisplayName = "Get Length Of Spline", Category = "Control Rig"))
+USTRUCT(meta = (DisplayName = "Get Length Of Spline", Category = "Spline"))
 struct CONTROLRIGSPLINE_API FRigUnit_GetLengthControlRigSpline : public FRigUnit
 {
 	GENERATED_BODY()
@@ -494,7 +494,7 @@ struct CONTROLRIGSPLINE_API FRigUnit_FitSplineCurveToChainItemArray : public FRi
 	FControlRigSpline Spline;	
 };
 
-USTRUCT(meta = (DisplayName = "Closest Parameter From Spline", Category = "Control Rig"))
+USTRUCT(meta = (DisplayName = "Closest Parameter From Spline", Category = "Spline"))
 struct CONTROLRIGSPLINE_API FRigUnit_ClosestParameterFromControlRigSpline : public FRigUnit_ControlRigSplineBase
 {
 	GENERATED_BODY()
