Control Rig Spline: Remove error on init of SetPointsOnSpline

diff --git a/Engine/Plugins/Animation/ControlRigSpline/Source/ControlRigSpline/Private/ControlRigSplineUnits.cpp b/Engine/Plugins/Animation/ControlRigSpline/Source/ControlRigSpline/Private/ControlRigSplineUnits.cpp
--- a/Engine/Plugins/Animation/ControlRigSpline/Source/ControlRigSpline/Private/ControlRigSplineUnits.cpp
+++ b/Engine/Plugins/Animation/ControlRigSpline/Source/ControlRigSpline/Private/ControlRigSplineUnits.cpp
@@ -36,6 +36,11 @@ FRigUnit_ControlRigSplineFromPoints_Execute()
 
 FRigUnit_SetSplinePoints_Execute()
 {
+	if (Context.State == EControlRigState::Init)
+	{
+		return;
+	}
+	
 	if (!Spline.SplineData.IsValid())
 	{
 		UE_LOG(LogControlRig, Error, TEXT("Invalid input spline."));
