Control Rig Spline: Add Param at percentage unit
#prefligth 6172d49cbe0cfd0001254379

diff --git a/Engine/Plugins/Animation/ControlRigSpline/Source/ControlRigSpline/Private/ControlRigSplineUnits.cpp b/Engine/Plugins/Animation/ControlRigSpline/Source/ControlRigSpline/Private/ControlRigSplineUnits.cpp
--- a/Engine/Plugins/Animation/ControlRigSpline/Source/ControlRigSpline/Private/ControlRigSplineUnits.cpp
+++ b/Engine/Plugins/Animation/ControlRigSpline/Source/ControlRigSpline/Private/ControlRigSplineUnits.cpp
@@ -3,6 +3,7 @@
 #include "ControlRigSplineUnits.h"
 #include "Units/RigUnitContext.h"
 #include "Features/IModularFeatures.h"
+#include "Algo/BinarySearch.h"
 
 #include "tinysplinecxx.h"
 
@@ -782,3 +783,55 @@ FRigUnit_ClosestParameterFromControlRigSpline_Execute()
 		}
 	}
 }
+
+FRigUnit_ParameterAtPercentage_Execute()
+{
+	if (!Spline.SplineData.IsValid())
+	{
+		return;
+	}
+	
+	switch (Context.State)
+	{
+	case EControlRigState::Init:
+	case EControlRigState::Update:
+		{
+			const int32 SampleCount = Spline.SplineData->SamplesArray.Num();
+
+			float Length = Spline.SplineData->AccumulatedLenth.Last();
+			float ClampedPercentage = FMath::Clamp(Percentage, 0.f, 1.f);
+
+			float SearchLength = Length * ClampedPercentage;
+			int32 NextIndex = Algo::LowerBound(Spline.SplineData->AccumulatedLenth, SearchLength);
+
+			if (NextIndex >= SampleCount)
+			{
+				U = 1.f;
+				return;
+			}
+
+			if (NextIndex <= 0)
+			{
+				U = 0.f;
+				return;
+			}
+
+			float UPrev = (NextIndex - 1) / (float) (SampleCount - 1);
+			float UNext = NextIndex / (float) (SampleCount - 1);
+
+			float LengthPrev = Spline.SplineData->AccumulatedLenth[NextIndex-1];
+			float LengthNext = Spline.SplineData->AccumulatedLenth[NextIndex];
+
+			float Interp = (SearchLength - LengthPrev) / (LengthNext - LengthPrev);
+			U = (1 - Interp) * UPrev + Interp * UNext;			
+				
+			break;
+		}
+	default:
+		{
+			checkNoEntry(); // Execute is only defined for Init and Update
+			break;
+		}
+	}
+}
+
diff --git a/Engine/Plugins/Animation/ControlRigSpline/Source/ControlRigSpline/Public/ControlRigSplineUnits.h b/Engine/Plugins/Animation/ControlRigSpline/Source/ControlRigSpline/Public/ControlRigSplineUnits.h
--- a/Engine/Plugins/Animation/ControlRigSpline/Source/ControlRigSpline/Public/ControlRigSplineUnits.h
+++ b/Engine/Plugins/Animation/ControlRigSpline/Source/ControlRigSpline/Public/ControlRigSplineUnits.h
@@ -520,3 +520,29 @@ struct CONTROLRIGSPLINE_API FRigUnit_ClosestParameterFromControlRigSpline : publ
 	
 };
 
+USTRUCT(meta = (DisplayName = "Parameter At Length Percentage", Category = "Spline"))
+struct CONTROLRIGSPLINE_API FRigUnit_ParameterAtPercentage : public FRigUnit_ControlRigSplineBase
+{
+	GENERATED_BODY()
+
+	FRigUnit_ParameterAtPercentage()
+	{
+		Percentage = 0.f;
+		U = 0.f;
+	}
+
+	/** Execute logic for this rig unit */
+	RIGVM_METHOD()
+	virtual void Execute(const FRigUnitContext& Context) override;
+
+	UPROPERTY(meta = (Input))
+	FControlRigSpline Spline;
+
+	UPROPERTY(meta = (Input))
+	float Percentage;
+
+	UPROPERTY(meta = (Output))
+	float U;
+	
+};
+
