Control Rig: Update node tooltips

diff --git a/Engine/Plugins/Animation/ControlRig/Source/ControlRig/Private/Units/Hierarchy/RigUnit_ProjectTransformToNewParent.h b/Engine/Plugins/Animation/ControlRig/Source/ControlRig/Private/Units/Hierarchy/RigUnit_ProjectTransformToNewParent.h
--- a/Engine/Plugins/Animation/ControlRig/Source/ControlRig/Private/Units/Hierarchy/RigUnit_ProjectTransformToNewParent.h
+++ b/Engine/Plugins/Animation/ControlRig/Source/ControlRig/Private/Units/Hierarchy/RigUnit_ProjectTransformToNewParent.h
@@ -6,12 +6,7 @@
 #include "RigUnit_ProjectTransformToNewParent.generated.h"
 
 /**
- * Projects a child from a previous to a new parent. The child doesn't
- * have to be parented that way in the hierarchy however.
- * For example: If you want the pelvis to move with the head (even though
- * the hierarchical relationship is the opposite) you can set the child
- * the pelvis (initial), set the old parent to be the head (initial)
- * and set the new parent to the head as well (not-initial == current).
+ * Gets the relative offset between the child and the old parent, then multiplies by new parent's transform.
  */
 USTRUCT(meta=(DisplayName="Project to new Parent", Category="Hierarchy", DocumentationPolicy = "Strict", Keywords = "ProjectTransformToNewParent,Relative,Reparent,Offset", NodeColor="0.462745, 1,0, 0.329412", Varying))
 struct CONTROLRIG_API FRigUnit_ProjectTransformToNewParent : public FRigUnit
diff --git a/Engine/Plugins/Animation/ControlRig/Source/ControlRig/Private/Units/Highlevel/Harmonics/RigUnit_BoneHarmonics.h b/Engine/Plugins/Animation/ControlRig/Source/ControlRig/Private/Units/Highlevel/Harmonics/RigUnit_BoneHarmonics.h
--- a/Engine/Plugins/Animation/ControlRig/Source/ControlRig/Private/Units/Highlevel/Harmonics/RigUnit_BoneHarmonics.h
+++ b/Engine/Plugins/Animation/ControlRig/Source/ControlRig/Private/Units/Highlevel/Harmonics/RigUnit_BoneHarmonics.h
@@ -143,7 +143,7 @@ struct CONTROLRIG_API FRigUnit_BoneHarmonics : public FRigUnit_HighlevelBaseMuta
 };
 
 /**
- * Performs point based simulation
+ * Drives an array of items through a harmonics spectrum
  */
 USTRUCT(meta=(DisplayName="Harmonics", Keywords="Sin,Wave"))
 struct CONTROLRIG_API FRigUnit_ItemHarmonics : public FRigUnit_HighlevelBaseMutable
diff --git a/Engine/Plugins/Animation/ControlRig/Source/ControlRig/Private/Units/Highlevel/Harmonics/RigUnit_ChainHarmonics.h b/Engine/Plugins/Animation/ControlRig/Source/ControlRig/Private/Units/Highlevel/Harmonics/RigUnit_ChainHarmonics.h
--- a/Engine/Plugins/Animation/ControlRig/Source/ControlRig/Private/Units/Highlevel/Harmonics/RigUnit_ChainHarmonics.h
+++ b/Engine/Plugins/Animation/ControlRig/Source/ControlRig/Private/Units/Highlevel/Harmonics/RigUnit_ChainHarmonics.h
@@ -229,7 +229,7 @@ struct CONTROLRIG_API FRigUnit_ChainHarmonics : public FRigUnit_HighlevelBaseMut
 };
 
 /**
- * Performs chain based simulation
+ * Given a root will drive all items underneath in a chain base harmonics spectrum
  */
 USTRUCT(meta=(DisplayName="Chain Harmonics"))
 struct CONTROLRIG_API FRigUnit_ChainHarmonicsPerItem : public FRigUnit_HighlevelBaseMutable
diff --git a/Engine/Plugins/Animation/ControlRig/Source/ControlRig/Private/Units/Simulation/RigUnit_AlphaInterp.h b/Engine/Plugins/Animation/ControlRig/Source/ControlRig/Private/Units/Simulation/RigUnit_AlphaInterp.h
--- a/Engine/Plugins/Animation/ControlRig/Source/ControlRig/Private/Units/Simulation/RigUnit_AlphaInterp.h
+++ b/Engine/Plugins/Animation/ControlRig/Source/ControlRig/Private/Units/Simulation/RigUnit_AlphaInterp.h
@@ -7,7 +7,7 @@
 #include "RigUnit_AlphaInterp.generated.h"
 
 /**
- * Adds a float value over time over and over again
+ * Takes in a float value and outputs an accumulated value with a customized scale and clamp
  */
 USTRUCT(meta=(DisplayName="Alpha Interpolate (float)", Keywords="Alpha,Lerp,LinearInterpolate", PrototypeName = "AlphaInterp"))
 struct CONTROLRIG_API FRigUnit_AlphaInterp : public FRigUnit_SimBase
@@ -80,7 +80,7 @@ struct CONTROLRIG_API FRigUnit_AlphaInterp : public FRigUnit_SimBase
 };
 
 /**
- * Adds a vector value over time over and over again
+ * Takes in a vector value and outputs an accumulated value with a customized scale and clamp
  */
 USTRUCT(meta=(DisplayName="Alpha Interpolate (Vector)", Keywords="Alpha,Lerp,LinearInterpolate", PrototypeName = "AlphaInterp", MenuDescSuffix = "(Vector)"))
 struct CONTROLRIG_API FRigUnit_AlphaInterpVector : public FRigUnit_SimBase
diff --git a/Engine/Plugins/Animation/ControlRigSpline/Source/ControlRigSpline/Public/ControlRigSplineUnits.h b/Engine/Plugins/Animation/ControlRigSpline/Source/ControlRigSpline/Public/ControlRigSplineUnits.h
--- a/Engine/Plugins/Animation/ControlRigSpline/Source/ControlRigSpline/Public/ControlRigSplineUnits.h
+++ b/Engine/Plugins/Animation/ControlRigSpline/Source/ControlRigSpline/Public/ControlRigSplineUnits.h
@@ -14,6 +14,9 @@ struct CONTROLRIGSPLINE_API FRigUnit_ControlRigSplineBase : public FRigUnit
 	GENERATED_BODY()
 };
 
+/*
+ * Creates a Spline curve from an array of positions
+ */
 USTRUCT(meta = (DisplayName = "Spline From Points", Keywords="Spline From Positions", Category = "Spline", Varying, NodeColor="0.737911 0.099899 0.099899"))
 struct CONTROLRIGSPLINE_API FRigUnit_ControlRigSplineFromPoints : public FRigUnit_ControlRigSplineBase
 {
@@ -50,6 +53,9 @@ struct CONTROLRIGSPLINE_API FRigUnit_ControlRigSplineFromPoints : public FRigUni
 	FControlRigSpline Spline;
 };
 
+/*
+ * Set the points of a spline, given a spline and an array of positions
+ */
 USTRUCT(meta = (DisplayName = "Set Spline Points", Category = "Spline", Varying, NodeColor = "0.0 0.36470600962638855 1.0"))
 struct CONTROLRIGSPLINE_API FRigUnit_SetSplinePoints : public FRigUnitMutable
 {
@@ -71,6 +77,9 @@ struct CONTROLRIGSPLINE_API FRigUnit_SetSplinePoints : public FRigUnitMutable
 	FControlRigSpline Spline;
 };
 
+/*
+ * Retrieves the position from a given Spline and U value
+ */
 USTRUCT(meta = (DisplayName = "Position From Spline", Keywords="Point From Spline", Category = "Spline", NodeColor="0.737911 0.099899 0.099899"))
 struct CONTROLRIGSPLINE_API FRigUnit_PositionFromControlRigSpline : public FRigUnit_ControlRigSplineBase
 {
@@ -96,6 +105,9 @@ struct CONTROLRIGSPLINE_API FRigUnit_PositionFromControlRigSpline : public FRigU
 	FVector Position;
 };
 
+/*
+ * Retrieves the transform from a given Spline and U value based on the given Up Vector and Roll
+ */
 USTRUCT(meta = (DisplayName = "Transform From Spline", Category = "Spline", NodeColor="0.737911 0.099899 0.099899"))
 struct CONTROLRIGSPLINE_API FRigUnit_TransformFromControlRigSpline : public FRigUnit_ControlRigSplineBase
 {
@@ -129,6 +141,9 @@ struct CONTROLRIGSPLINE_API FRigUnit_TransformFromControlRigSpline : public FRig
 	FTransform Transform;
 };
 
+/*
+ * Retrieves the tangent from a given Spline and U value
+ */
 USTRUCT(meta = (DisplayName = "Tangent From Spline", Category = "Spline", NodeColor="0.737911 0.099899 0.099899"))
 struct CONTROLRIGSPLINE_API FRigUnit_TangentFromControlRigSpline : public FRigUnit_ControlRigSplineBase
 {
@@ -154,6 +169,9 @@ struct CONTROLRIGSPLINE_API FRigUnit_TangentFromControlRigSpline : public FRigUn
 	FVector Tangent;
 };
 
+/*
+ * Draws the given spline in the viewport
+ */
 USTRUCT(meta = (DisplayName = "Draw Spline", Category = "Spline", NodeColor = "0.83077 0.846873 0.049707"))
 struct CONTROLRIGSPLINE_API FRigUnit_DrawControlRigSpline : public FRigUnitMutable
 {
@@ -183,6 +201,9 @@ struct CONTROLRIGSPLINE_API FRigUnit_DrawControlRigSpline : public FRigUnitMutab
 	int32 Detail;
 };
 
+/*
+ * Retrieves the length from a given Splin
+ */
 USTRUCT(meta = (DisplayName = "Get Length Of Spline", Category = "Spline", NodeColor="0.737911 0.099899 0.099899"))
 struct CONTROLRIGSPLINE_API FRigUnit_GetLengthControlRigSpline : public FRigUnit
 {
@@ -494,6 +515,9 @@ struct CONTROLRIGSPLINE_API FRigUnit_FitSplineCurveToChainItemArray : public FRi
 	FControlRigSpline Spline;	
 };
 
+/*
+ * Retrieves the closest U value from a given Spline and a position
+ */
 USTRUCT(meta = (DisplayName = "Closest Parameter From Spline", Category = "Spline", NodeColor="0.737911 0.099899 0.099899"))
 struct CONTROLRIGSPLINE_API FRigUnit_ClosestParameterFromControlRigSpline : public FRigUnit_ControlRigSplineBase
 {
@@ -520,6 +544,9 @@ struct CONTROLRIGSPLINE_API FRigUnit_ClosestParameterFromControlRigSpline : publ
 	
 };
 
+/*
+ * Returns the U parameter of a spline given a length percentage (0.0 - 1.0)
+ */
 USTRUCT(meta = (DisplayName = "Parameter At Length Percentage", Category = "Spline"))
 struct CONTROLRIGSPLINE_API FRigUnit_ParameterAtPercentage : public FRigUnit_ControlRigSplineBase
 {
