Control Rig Spline: Improve performance of hermite spline

diff --git a/Engine/Plugins/Animation/ControlRigSpline/Source/ControlRigSpline/Private/ControlRigSplineTypes.cpp b/Engine/Plugins/Animation/ControlRigSpline/Source/ControlRigSpline/Private/ControlRigSplineTypes.cpp
--- a/Engine/Plugins/Animation/ControlRigSpline/Source/ControlRigSpline/Private/ControlRigSplineTypes.cpp
+++ b/Engine/Plugins/Animation/ControlRigSpline/Source/ControlRigSpline/Private/ControlRigSplineTypes.cpp
@@ -74,32 +74,6 @@ void FControlRigSpline::SetControlPoints(const TArrayView<const FVector>& InPoin
 			}
 			case ESplineType::Hermite:
 			{
-				// With this spline type, we cannot just move the contorl points around, since the interpolation creates
-				// new control points from the ones provided. The interpolation allocates a new spline each time
-				// its called.
-						
-				// Create the control points array for the interpolation.
-				// The first dimension acts as the query parameter, the rest of the dimensions are the XYZ of the control points
-#if (UE_LARGE_WORLD_COORDINATES_DISABLED && !TINYSPLINE_FLOAT_PRECISION)					
-				TArray<double> ControlPointsArray;
-#else
-				TArray<FVector::FReal> ControlPointsArray;
-#endif
-				ControlPointsArray.SetNumUninitialized(ControlPointsCount * 4, false);
-	
-				for (int32 i = 0; i < ControlPointsCount; ++i)
-				{
-					ControlPointsArray[i * 4 + 0] = i / (FVector::FReal)(ControlPointsCount - 1);
-
-					ControlPointsArray[i * 4 + 1] = (FVector::FReal)InPoints[i][0];
-					ControlPointsArray[i * 4 + 2] = (FVector::FReal)InPoints[i][1];
-					ControlPointsArray[i * 4 + 3] = (FVector::FReal)InPoints[i][2];
-				}
-
-				// Create a new spline to store the result of the interpolation
-				ts_bspline_free(SplineData->Spline.data());					
-				ts_bspline_interpolate_cubic_natural(ControlPointsArray.GetData(), ControlPointsCount, 4, SplineData->Spline.data(), nullptr);
-				
 				break;
 			}
 			default:
@@ -143,28 +117,39 @@ void FControlRigSpline::SetControlPoints(const TArrayView<const FVector>& InPoin
 			}
 			case ESplineType::Hermite:
 			{
-				// Cache sample positions of the spline
-#if (UE_LARGE_WORLD_COORDINATES_DISABLED && !TINYSPLINE_FLOAT_PRECISION)					
-				double* SamplesPtr = nullptr;
-#else
-				FVector::FReal* SamplesPtr = nullptr;
-#endif
-				size_t ActualSamplesPerSegment = 0;
-				ts_bspline_sample(SplineData->Spline.data(), (ControlPointsCount - 1) * SamplesPerSegment, &SamplesPtr, &ActualSamplesPerSegment, nullptr);
 				SplineData->SamplesArray.SetNumUninitialized((ControlPointsCount - 1) * SamplesPerSegment);
 				for (int32 i = 0; i < ControlPointsCount - 1; ++i)
 				{
+					const FVector P0 = (i > 0) ? InPoints[i-1] : 2*InPoints[0] - InPoints[1];
+					const FVector& P1 = InPoints[i];
+					const FVector& P2 = InPoints[i+1];
+					const FVector P3 = (i + 2 < ControlPointsCount) ? InPoints[i+2] : 2*InPoints.Last() - InPoints[ControlPointsCount-2];
+
+					// https://www.cs.cmu.edu/~fp/courses/graphics/asst5/catmullRom.pdf
+					float Tension = 0.5f;
+					FVector M1 = Tension * (P2 - P0);
+					FVector M2 = Tension * (P3 - P1);
+
+					float Dt = 1.f / (float) SamplesPerSegment;
+					if (i == ControlPointsCount - 2)
+					{
+						Dt = 1.f / (float) (SamplesPerSegment - 1);
+					}
 					for (int32 j = 0; j < SamplesPerSegment; ++j)
 					{
-						SplineData->SamplesArray[i * SamplesPerSegment + j].X = SamplesPtr[(i * SamplesPerSegment + j) * 4 + 1];
-						SplineData->SamplesArray[i * SamplesPerSegment + j].Y = SamplesPtr[(i * SamplesPerSegment + j) * 4 + 2];
-						SplineData->SamplesArray[i * SamplesPerSegment + j].Z = SamplesPtr[(i * SamplesPerSegment + j) * 4 + 3];
+						// https://en.wikipedia.org/wiki/Cubic_Hermite_spline#Catmull–Rom_spline
+						const float T = j  * Dt;
+						const float T2 = T*T;
+						const float T3 = T2*T;
+	
+						float H00 = 2*T3 - 3*T2 + 1;
+						float H10 = T3 - 2*T2 + T;
+						float H01 = -2*T3 + 3*T2;
+						float H11 = T3 - T2;
+
+						SplineData->SamplesArray[i * SamplesPerSegment + j] = H00*P1 + H10*M1 + H01*P2 + H11*M2;						
 					}
 				}
-
-				// tinySpline will allocate the samples array, but does not free that memory. We need to take care of that.
-				free(SamplesPtr);
-
 				break;
 			}
 			default:
@@ -222,49 +207,6 @@ void FControlRigSpline::SetControlPoints(const TArrayView<const FVector>& InPoin
 	}	
 }
 
-int32 FControlRigSpline::GetControlPoints(TArray<FVector>& OutPoints) const
-{
-	const size_t Count = SplineData->Spline.numControlPoints();
-	OutPoints.SetNumUninitialized(Count);
-
-#if (UE_LARGE_WORLD_COORDINATES_DISABLED && !TINYSPLINE_FLOAT_PRECISION)					
-	double* Points;
-#else
-	FVector::FReal* Points;
-#endif	
-	
-	ts_bspline_control_points(SplineData->Spline.data(), &Points, nullptr);
-
-	switch (SplineData->SplineMode)
-	{
-		case ESplineType::BSpline:
-		{
-			for (int32 i = 0; i < Count; ++i)
-			{
-				int32 index = i*3;
-				OutPoints[i].Set(Points[index + 0], Points[index + 1], Points[index + 2]);
-			}
-			break;
-		}
-		case ESplineType::Hermite:
-		{
-			for (int32 i = 0; i < Count; ++i)
-			{
-				int32 index = i * 4;
-				OutPoints[i].Set(Points[index + 1], Points[index + 2], Points[index + 3]);
-			}
-			break;
-		}
-		default:
-		{
-			checkNoEntry();
-			break;
-		}
-	}
-
-	return Count;
-}
-
 FVector FControlRigSpline::PositionAtParam(const float InParam) const
 {
 	if (!SplineData.IsValid())
diff --git a/Engine/Plugins/Animation/ControlRigSpline/Source/ControlRigSpline/Public/ControlRigSplineTypes.h b/Engine/Plugins/Animation/ControlRigSpline/Source/ControlRigSpline/Public/ControlRigSplineTypes.h
--- a/Engine/Plugins/Animation/ControlRigSpline/Source/ControlRigSpline/Public/ControlRigSplineTypes.h
+++ b/Engine/Plugins/Animation/ControlRigSpline/Source/ControlRigSpline/Public/ControlRigSplineTypes.h
@@ -90,14 +90,6 @@ struct CONTROLRIGSPLINE_API FControlRigSpline
 	*/
 	void SetControlPoints(const TArrayView<const FVector>& InPoints, const ESplineType SplineMode = ESplineType::BSpline, const int32 SamplesPerSegment = 16, const float Compression = 1.f, const float Stretch = 1.f);
 
-	/**
-	* Populates OutPoints and returns the number of control points in the spline.
-	*
-	* @param OutPoints	The array to populate.
-	* @return			The number of control points.
-	*/
-	int32 GetControlPoints(TArray<FVector>& OutPoints) const;
-
 	/**
 	* Given an InParam float in [0, 1], will return the position of the spline at that point.
 	*
