Control Rig: Allow content on control rig spline

diff --git a/Engine/Plugins/Animation/ControlRigSpline/ControlRigSpline.uplugin b/Engine/Plugins/Animation/ControlRigSpline/ControlRigSpline.uplugin
--- a/Engine/Plugins/Animation/ControlRigSpline/ControlRigSpline.uplugin
+++ b/Engine/Plugins/Animation/ControlRigSpline/ControlRigSpline.uplugin
@@ -11,7 +11,7 @@
 	"MarketplaceURL": "",
 	"SupportURL": "",
 	"EnabledByDefault": true,
-	"CanContainContent": false,
+	"CanContainContent": true,
 	"IsBetaVersion": true,
 	"IsExperimentalVersion":  true,
 	"Installed": false,
