Control Rig: Add native splines for future TinySpline dependency removal

diff --git a/Engine/Plugins/Animation/ControlRigSpline/ControlRigSpline.uplugin b/Engine/Plugins/Animation/ControlRigSpline/ControlRigSpline.uplugin
--- a/Engine/Plugins/Animation/ControlRigSpline/ControlRigSpline.uplugin
+++ b/Engine/Plugins/Animation/ControlRigSpline/ControlRigSpline.uplugin
@@ -26,11 +26,6 @@
 			"Name": "ControlRigSpline",
 			"Type": "Runtime",
 			"LoadingPhase": "Default"
-		},
-		{
-			"Name": "ControlRigSplineEditor",
-			"Type": "Editor",
-			"LoadingPhase": "PostEngineInit"
 		}
 	]
 }
\ No newline at end of file
diff --git a/Engine/Plugins/Animation/ControlRigSpline/Source/ControlRigSpline/Private/ControlRigSplineTypes.cpp b/Engine/Plugins/Animation/ControlRigSpline/Source/ControlRigSpline/Private/ControlRigSplineTypes.cpp
--- a/Engine/Plugins/Animation/ControlRigSpline/Source/ControlRigSpline/Private/ControlRigSplineTypes.cpp
+++ b/Engine/Plugins/Animation/ControlRigSpline/Source/ControlRigSpline/Private/ControlRigSplineTypes.cpp
@@ -2,6 +2,195 @@
 
 #include "ControlRigSplineTypes.h"
 
+#if !(USE_TINYSPLINE)
+
+ControlRigBaseSpline::ControlRigBaseSpline(const TArrayView<const FVector>& InControlPoints, const uint16 InDegree)
+: ControlPoints(InControlPoints)
+, Degree(InDegree)
+{
+	
+}
+
+ControlRigBSpline::ControlRigBSpline(const TArrayView<const FVector>& InControlPoints, const uint16 InDegree, bool bInClamped)
+: ControlRigBaseSpline(InControlPoints, InDegree)
+{
+	uint16 N = ControlPoints.Num() - 1;
+	uint16 P = Degree;
+	uint16 M = N + P + 1;
+	KnotVector.SetNumZeroed(M+1);
+
+	// Clamped means there is p+1 0s at the beginning and 1s at the end of the knot vector
+	if (bInClamped)
+	{
+		for (int32 i=M-P; i<=M; ++i)
+		{
+			KnotVector[i] = 1.f;
+		}
+
+		float Dk = 1.f/(M-P - (P+1) +1);
+		for (int32 i=P+1; i<M-P; ++i)
+		{
+			KnotVector[i] = (i-P)*Dk;
+		}
+	}
+	else
+	{
+		float Dk = 1.f / M;
+		for (int32 i = 0; i<M+1; ++i)
+		{
+			KnotVector[i] = i*Dk;
+		}
+	}
+}
+
+FVector ControlRigBSpline::GetPointAtParam(float Param)
+{
+	// DeBoor's Algorithm
+	// https://pages.mtu.edu/~shene/COURSES/cs3621/NOTES/spline/B-spline/de-Boor.html
+	
+	float ClampedU = FMath::Clamp(Param, 0.f, 1.f);
+
+	uint16 N = ControlPoints.Num() - 1;
+	uint16 P = Degree;
+	uint16 M = N + P + 1;
+
+	// Remap U value to KnotVector[P] : KnotVector[M+1 - (P+1)]
+	float Min = KnotVector[P] + KINDA_SMALL_NUMBER;
+	float Max = KnotVector[M-P] - KINDA_SMALL_NUMBER;
+	float U = Min + ClampedU*(Max-Min);
+
+	// Find [uk, uk+1) where U lives (or k == the first appearence of U)
+	uint16 K = 0;
+	{
+		// Find [uk, uk+1) where U lives
+		for (uint16 i=0; i<KnotVector.Num()-1; ++i)
+		{
+			if ((FMath::IsNearlyEqual(KnotVector[i], U)) ||
+				(KnotVector[i] < U && KnotVector[i+1] > U))
+			{
+				K = i;
+				break;
+			}
+		}
+	}
+
+	TArray<FVector> WorkingPoints;
+	WorkingPoints.SetNum(P+1);
+	for (int32 i=K-P; i<=K; ++i)
+	{
+		WorkingPoints[i-(K-P)] = ControlPoints[i];
+	}
+	
+	for (uint16 R = 1; R <= Degree; ++R)
+	{
+		uint16 StartCPIndex = K-P+R;
+		uint16 EndCPIndex = K;
+		if (FMath::IsNearlyEqual(KnotVector[StartCPIndex], U))
+		{
+			continue;
+		}
+
+		uint16 Index = 0;
+		for (uint16 i = StartCPIndex; i <= EndCPIndex; ++i, ++Index)
+		{
+			float Air = (U - KnotVector[i]) / (KnotVector[i+P-R+1] - KnotVector[i]);
+			WorkingPoints[Index] = (1-Air) * WorkingPoints[Index] + Air * WorkingPoints[Index+1];
+		}
+	}
+
+	return WorkingPoints[0];
+}
+
+ControlRigHermite::ControlRigHermite(const TArrayView<const FVector>& InControlPoints)
+: ControlRigBaseSpline(InControlPoints, 3)
+{
+	SegmentPoints.Add(ControlPoints[0] + ControlPoints[0] - ControlPoints[1]);
+	SegmentPoints.Append(ControlPoints);
+	SegmentPoints.Add(ControlPoints.Last() + ControlPoints.Last() - ControlPoints[ControlPoints.Num()-2]);
+
+	NumSegments = SegmentPoints.Num() - 3;
+}
+
+void ControlRigHermite::SetControlPoints(const TArrayView<const FVector>& InControlPoints)
+{
+	ControlRigBaseSpline::SetControlPoints(InControlPoints);
+
+	SegmentPoints.Reset();
+	SegmentPoints.Add(ControlPoints[0] + ControlPoints[0] - ControlPoints[1]);
+	SegmentPoints.Append(ControlPoints);
+	SegmentPoints.Add(ControlPoints.Last() + ControlPoints.Last() - ControlPoints[ControlPoints.Num()-2]);
+}
+
+static FVector Hermite(const FVector& P0, const FVector& P1, const FVector& M0, const FVector& M1, float Param)
+{
+	// https://en.wikipedia.org/wiki/Cubic_Hermite_spline#Catmull–Rom_spline
+
+	const float T = Param;
+	const float T2 = T*T;
+	const float T3 = T2*T;
+	
+	float H00 = 2*T3 - 3*T2 + 1;
+	float H10 = T3 - 2*T2 + T;
+	float H01 = -2*T3 + 3*T2;
+	float H11 = T3 - T2;
+
+	return H00*P0 + H10*M0 + H01*P1 + H11*M1;
+}
+
+FVector ControlRigHermite::GetPointAtParam(float Param)
+{
+	float ClampedParam = FMath::Clamp(Param, 0.f, 1.f);	
+
+	float ParamIntervals = 1.f / (float) NumSegments;
+	uint16 SegmentIndex = FMath::Floor(ClampedParam / ParamIntervals);
+	if (SegmentIndex >= NumSegments)
+	{
+		SegmentIndex = NumSegments-1;
+	}
+
+	uint16 P3Index = SegmentIndex + 3;
+	P3Index = (P3Index < SegmentPoints.Num()) ? P3Index : SegmentPoints.Num()-1;
+	uint16 P0Index = P3Index - 3;
+	uint16 P1Index = P3Index - 2;
+	uint16 P2Index = P3Index - 1;
+
+	float StartSegmentParam = SegmentIndex * ParamIntervals;
+	float T = FMath::Clamp((ClampedParam-StartSegmentParam)/ParamIntervals, 0.f, 1.f);
+
+	const FVector& P0 = SegmentPoints[P0Index];
+	const FVector& P1 = SegmentPoints[P1Index];
+	const FVector& P2 = SegmentPoints[P2Index];
+	const FVector& P3 = SegmentPoints[P3Index];
+
+	// https://www.cs.cmu.edu/~fp/courses/graphics/asst5/catmullRom.pdf
+	float Tension = 0.5f;
+	FVector M1 = Tension * (P2 - P0);
+	FVector M2 = Tension * (P3 - P1);
+	return Hermite(P1, P2, M1, M2, T);
+}
+#endif
+
+FControlRigSplineImpl::~FControlRigSplineImpl()
+{
+#if !(USE_TINYSPLINE)
+	if (Spline)
+	{
+		delete Spline;
+		Spline = nullptr;
+	}
+#endif
+}
+
+TArray<FVector>& FControlRigSplineImpl::GetControlPoints()
+{
+#if USE_TINYSPLINE
+	return ControlPoints;
+#else
+	check(Spline);
+	return Spline->GetControlPoints();	
+#endif
+}
+
 FControlRigSpline::FControlRigSpline(const FControlRigSpline& InOther)
 {
 	SplineData = InOther.SplineData;
@@ -26,17 +215,24 @@ void FControlRigSpline::SetControlPoints(const TArrayView<const FVector>& InPoin
 		SplineData = MakeShared<FControlRigSplineImpl>();
 	}
 
-	bool bControlPointsChanged = InPoints != SplineData->ControlPoints;
+#if USE_TINYSPLINE
+	bool bControlPointsChanged = InPoints != SplineData->GetControlPoints();
+	bool bNumControlPointsChanged = SplineData->GetControlPoints().Num() != ControlPointsCount;
+#else
+	bool bControlPointsChanged = (SplineData->Spline) ? InPoints != SplineData->GetControlPoints() : true;
+	bool bNumControlPointsChanged = (SplineData->Spline) ? SplineData->GetControlPoints().Num() != ControlPointsCount : true;
+#endif
 	bool bSplineModeChanged = SplineMode != SplineData->SplineMode;
 	bool bSamplesCountChanged = SamplesPerSegment != SplineData->SamplesPerSegment;
-	bool bNumControlPointsChanged = SplineData->ControlPoints.Num() != ControlPointsCount;
 	bool bStretchChanged = Stretch != SplineData->Stretch || Compression != SplineData->Compression;
 	if (!bSplineModeChanged && !bControlPointsChanged && !bSamplesCountChanged && !bStretchChanged)
 	{
 		return;
 	}
 
+#if USE_TINYSPLINE
 	SplineData->ControlPoints = InPoints;
+#endif
 	SplineData->SplineMode = SplineMode;
 	SplineData->SamplesPerSegment = SamplesPerSegment;
 	SplineData->Compression = Compression;
@@ -45,6 +241,7 @@ void FControlRigSpline::SetControlPoints(const TArrayView<const FVector>& InPoin
 	// If we need to update the spline because the controls points have changed, or the spline mode has changed
 	if (bControlPointsChanged || bSplineModeChanged)
 	{	
+#if USE_TINYSPLINE
 		switch (SplineMode)
 		{
 			case ESplineType::BSpline:
@@ -77,11 +274,49 @@ void FControlRigSpline::SetControlPoints(const TArrayView<const FVector>& InPoin
 				break;
 			}
 		}
+
+#else
+
+		if (bNumControlPointsChanged || bSplineModeChanged)
+		{
+			// Delete previously created spline
+			if (SplineData->Spline)
+			{
+				delete SplineData->Spline;
+				SplineData->Spline = nullptr;
+			}
+			
+			switch (SplineMode)
+			{
+				case ESplineType::BSpline:
+				{
+					SplineData->Spline = new ControlRigBSpline(InPoints, 3, true);
+					break;
+				}
+				case ESplineType::Hermite:
+				{
+					
+					SplineData->Spline = new ControlRigHermite(InPoints);
+					break;
+				}
+				default:
+				{
+					checkNoEntry(); // Unknown Spline Mode
+					break;
+				}
+			}
+		}
+		else if (bControlPointsChanged)
+		{
+			SplineData->Spline->SetControlPoints(InPoints);
+		}
+#endif
 	}
 
 	// If curve has changed, or sample count has changed, recompute the cache
 	if (bControlPointsChanged || bSplineModeChanged || bSamplesCountChanged || bStretchChanged)
 	{
+#if USE_TINYSPLINE
 		switch (SplineMode)
 		{
 			case ESplineType::BSpline:
@@ -149,6 +384,18 @@ void FControlRigSpline::SetControlPoints(const TArrayView<const FVector>& InPoin
 				break;
 			}
 		}
+#else
+		// Cache sample positions of the spline
+		int32 NumSamples = (ControlPointsCount-1) * SamplesPerSegment;
+		SplineData->SamplesArray.SetNumUninitialized(NumSamples, false);
+			
+		float U = 0.f;
+		float DeltaU = 1.f/(NumSamples-1);
+		for (int32 i=0; i<NumSamples; ++i, U+=DeltaU)
+		{
+			SplineData->SamplesArray[i] = SplineData->Spline->GetPointAtParam(U); 
+		}
+#endif
 
 		// Correct length of samples
 		if (!bSplineModeChanged && !bSamplesCountChanged && !bNumControlPointsChanged)
diff --git a/Engine/Plugins/Animation/ControlRigSpline/Source/ControlRigSpline/Private/ControlRigSplineUnits.cpp b/Engine/Plugins/Animation/ControlRigSpline/Source/ControlRigSpline/Private/ControlRigSplineUnits.cpp
--- a/Engine/Plugins/Animation/ControlRigSpline/Source/ControlRigSpline/Private/ControlRigSplineUnits.cpp
+++ b/Engine/Plugins/Animation/ControlRigSpline/Source/ControlRigSpline/Private/ControlRigSplineUnits.cpp
@@ -48,7 +48,15 @@ FRigUnit_SetSplinePoints_Execute()
 		return;
 	}
 	
-	if (Points.Num() != Spline.SplineData->ControlPoints.Num())
+#if !(USE_TINYSPLINE)
+	if (Spline.SplineData->Spline == nullptr)
+	{
+		UE_LOG(LogControlRig, Error, TEXT("Invalid input spline implementation."));
+		return;
+	}	
+#endif
+	
+	if (Points.Num() != Spline.SplineData->GetControlPoints().Num())
 	{
 		UE_LOG(LogControlRig, Error, TEXT("Number of input points does not match the number of point in the spline."));
 		return;
@@ -77,6 +85,14 @@ FRigUnit_PositionFromControlRigSpline_Execute()
 	{
 		return;
 	}
+
+#if !(USE_TINYSPLINE)
+	if (Spline.SplineData->Spline == nullptr)
+	{
+		UE_LOG(LogControlRig, Error, TEXT("Invalid input spline implementation."));
+		return;
+	}
+#endif
 	
 	switch (Context.State)
 	{
@@ -100,6 +116,14 @@ FRigUnit_TransformFromControlRigSpline_Execute()
 	{
 		return;
 	}
+
+#if !(USE_TINYSPLINE)
+	if (Spline.SplineData->Spline == nullptr)
+	{
+		UE_LOG(LogControlRig, Error, TEXT("Invalid input spline implementation."));
+		return;
+	}
+#endif	
 	
 	switch (Context.State)
 	{
@@ -142,6 +166,14 @@ FRigUnit_TangentFromControlRigSpline_Execute()
 	{
 		return;
 	}
+
+#if !(USE_TINYSPLINE)
+	if (Spline.SplineData->Spline == nullptr)
+	{
+		UE_LOG(LogControlRig, Error, TEXT("Invalid input spline implementation."));
+		return;
+	}
+#endif	
 	
 	switch (Context.State)
 	{
@@ -189,6 +221,14 @@ FRigUnit_DrawControlRigSpline_Execute()
 		return;
 	}
 
+#if !(USE_TINYSPLINE)
+	if (Spline.SplineData->Spline == nullptr)
+	{
+		UE_LOG(LogControlRig, Error, TEXT("Invalid input spline implementation."));
+		return;
+	}
+#endif	
+
 	int32 Count = FMath::Clamp<int32>(Detail, 4, 64);
 	FControlRigDrawInstruction Instruction(EControlRigDrawSettings::LineStrip, Color, Thickness);
 	Instruction.Positions.SetNumUninitialized(Count);
@@ -212,6 +252,15 @@ FRigUnit_GetLengthControlRigSpline_Execute()
 		Length = 0;
 		return;
 	}
+
+#if !(USE_TINYSPLINE)
+	if (Spline.SplineData->Spline == nullptr)
+	{
+		Length = 0;
+		UE_LOG(LogControlRig, Error, TEXT("Invalid input spline implementation."));
+		return;
+	}
+#endif	
 	
 	switch (Context.State)
 	{
@@ -289,6 +338,14 @@ FRigUnit_FitChainToSplineCurveItemArray_Execute()
 		return;
 	}
 
+#if !(USE_TINYSPLINE)
+	if (Spline.SplineData->Spline == nullptr)
+	{
+		UE_LOG(LogControlRig, Error, TEXT("Invalid input spline implementation."));
+		return;
+	}
+#endif	
+
 	float& ChainLength = WorkData.ChainLength;
 	TArray<FVector>& ItemPositions = WorkData.ItemPositions;
 	TArray<float>& ItemSegments = WorkData.ItemSegments;
@@ -671,7 +728,7 @@ FRigUnit_FitChainToSplineCurveItemArray_Execute()
 		}
 		Context.DrawInterface->Instructions.Add(Instruction);
 
-		for (auto Point : Spline.SplineData->ControlPoints)
+		for (auto Point : Spline.SplineData->GetControlPoints())
 		{
 			Context.DrawInterface->DrawPoint(DebugSettings.WorldOffset, Point, DebugSettings.Scale * 6, DebugSettings.CurveColor);
 		}
@@ -710,6 +767,14 @@ FRigUnit_FitSplineCurveToChainItemArray_Execute()
 		return;
 	}
 
+#if !(USE_TINYSPLINE)
+	if (Spline.SplineData->Spline == nullptr)
+	{
+		UE_LOG(LogControlRig, Error, TEXT("Invalid input spline implementation."));
+		return;
+	}
+#endif	
+
 	if (Context.State == EControlRigState::Init)
 	{
 		return;
@@ -737,7 +802,7 @@ FRigUnit_FitSplineCurveToChainItemArray_Execute()
 	// 2.-  for each control point in the original spline
 	//			figure out its u in the original spline (preprocess)
 	//			query position at u on the new spline
-	const TArray<FVector>& ControlPoints = Spline.SplineData->ControlPoints;
+	const TArray<FVector>& ControlPoints = Spline.SplineData->GetControlPoints();
 	TArray<FVector> NewControlPoints;
 	NewControlPoints.SetNumUninitialized(ControlPoints.Num());
 	for (int32 i = 0; i < ControlPoints.Num(); ++i)
@@ -757,6 +822,14 @@ FRigUnit_ClosestParameterFromControlRigSpline_Execute()
 	{
 		return;
 	}
+
+#if !(USE_TINYSPLINE)
+	if (Spline.SplineData->Spline == nullptr)
+	{
+		UE_LOG(LogControlRig, Error, TEXT("Invalid input spline implementation."));
+		return;
+	}
+#endif	
 	
 	switch (Context.State)
 	{
@@ -821,6 +894,14 @@ FRigUnit_ParameterAtPercentage_Execute()
 	{
 		return;
 	}
+
+#if !(USE_TINYSPLINE)
+	if (Spline.SplineData->Spline == nullptr)
+	{
+		UE_LOG(LogControlRig, Error, TEXT("Invalid input spline implementation."));
+		return;
+	}
+#endif	
 	
 	switch (Context.State)
 	{
diff --git a/Engine/Plugins/Animation/ControlRigSpline/Source/ControlRigSpline/Public/ControlRigSplineTypes.h b/Engine/Plugins/Animation/ControlRigSpline/Source/ControlRigSpline/Public/ControlRigSplineTypes.h
--- a/Engine/Plugins/Animation/ControlRigSpline/Source/ControlRigSpline/Public/ControlRigSplineTypes.h
+++ b/Engine/Plugins/Animation/ControlRigSpline/Source/ControlRigSpline/Public/ControlRigSplineTypes.h
@@ -2,9 +2,15 @@
 
 #pragma once
 
+// Define to switch between spline using the thirdparty TinySpline implementation or
+// our native spline implementation
+#define USE_TINYSPLINE 0
+
 #include "CoreMinimal.h"
 #include "ControlRigDefines.h"
+#if USE_TINYSPLINE
 #include "tinysplinecxx.h"
+#endif
 #include "ControlRigSplineTypes.generated.h"
 
 class UControlRig;
@@ -24,6 +30,64 @@ enum class ESplineType : uint8
 	Max UMETA(Hidden),
 };
 
+#if !(USE_TINYSPLINE)
+// Reading material
+// 
+// https://pages.mtu.edu/~shene/COURSES/cs3621/NOTES/notes.html
+// [REDACTED_PROFILE_URL]/DeBoorAlgorithmNurbs
+// https://stackoverflow.com/questions/32897943/how-do-catmull-rom-and-hermite-splines-relate
+
+// https://stackoverflow.com/questions/30748316/%20catmull-rom-interpolation-on-svg-paths/30826434#30826434
+// https://stackoverflow.com/questions/32897943/how-do-catmull-rom-and-hermite-splines-relate
+// https://www.cs.cmu.edu/~fp/courses/graphics/asst5/catmullRom.pdf
+
+class ControlRigBaseSpline
+{
+
+protected:
+	TArray<FVector> ControlPoints;
+	uint16 Degree;
+	
+public:
+
+	ControlRigBaseSpline(const TArrayView<const FVector>& InControlPoints, uint16 InDegree);
+	virtual ~ControlRigBaseSpline(){}
+	
+	virtual FVector GetPointAtParam(float Param) = 0;
+
+	virtual void SetControlPoints(const TArrayView<const FVector>& InControlPoints) { ControlPoints = InControlPoints; }
+
+	TArray<FVector>& GetControlPoints() { return ControlPoints; }
+};
+
+class ControlRigBSpline : public ControlRigBaseSpline
+{
+	TArray<float> KnotVector;
+	
+public:
+	ControlRigBSpline(const TArrayView<const FVector>& InControlPoints, const uint16 InDegree, const bool bInClamped);
+	
+	virtual FVector GetPointAtParam(float Param);
+
+};
+
+class ControlRigHermite : public ControlRigBaseSpline
+{
+	TArray<FVector> SegmentPoints;
+
+	uint16 NumSegments;
+	
+public:
+	ControlRigHermite(const TArrayView<const FVector>& InControlPoints);
+
+	virtual void SetControlPoints(const TArrayView<const FVector>& InControlPoints) override;
+	
+	virtual FVector GetPointAtParam(float Param);
+
+};
+
+#endif
+
 USTRUCT()
 struct CONTROLRIGSPLINE_API FControlRigSplineImpl
 {
@@ -32,20 +96,31 @@ struct CONTROLRIGSPLINE_API FControlRigSplineImpl
 	FControlRigSplineImpl()	
 	{
 		SplineMode = ESplineType::BSpline;
+#if !(USE_TINYSPLINE)
+		Spline = nullptr;
+#endif
 		SamplesPerSegment = 16;
 	}
 
+	virtual ~FControlRigSplineImpl();
+
 	// Spline type
 	ESplineType SplineMode;
 
+#if USE_TINYSPLINE
 	// The control points to construct the spline
 	TArray<FVector> ControlPoints;
+#endif
 
 	// The initial lengths between samples
 	TArray<float> InitialLengths;
 
 	// The actual spline
+#if USE_TINYSPLINE
 	tinyspline::BSpline Spline;
+#else
+	ControlRigBaseSpline* Spline;
+#endif
 
 	// Samples per segment, where segment is the portion between two control points
 	int32 SamplesPerSegment;
@@ -61,6 +136,9 @@ struct CONTROLRIGSPLINE_API FControlRigSplineImpl
 
 	// Accumulated length along the spline given by samples
 	TArray<float> AccumulatedLenth;
+
+	// Returns a reference to the control points that were used to create this spline 
+	TArray<FVector>& GetControlPoints();
 };
 
 USTRUCT(BlueprintType)
diff --git a/Engine/Plugins/Animation/ControlRigSpline/Source/ControlRigSplineEditor/ControlRigSplineEditor.build.cs b/Engine/Plugins/Animation/ControlRigSpline/Source/ControlRigSplineEditor/ControlRigSplineEditor.build.cs
deleted file mode 100644
diff --git a/Engine/Plugins/Animation/ControlRigSpline/Source/ControlRigSplineEditor/Private/ControlRigSplineEditorModule.cpp b/Engine/Plugins/Animation/ControlRigSpline/Source/ControlRigSplineEditor/Private/ControlRigSplineEditorModule.cpp
deleted file mode 100644
diff --git a/Engine/Plugins/Animation/ControlRigSpline/Source/ControlRigSplineEditor/Public/ControlRigSplineEditorModule.h b/Engine/Plugins/Animation/ControlRigSpline/Source/ControlRigSplineEditor/Public/ControlRigSplineEditorModule.h
deleted file mode 100644
