ControlRigSpline: Handle private includes from public headers so the path isn't rooted in the module parent. This should be cleaned up to not require including private headers in public files.

#rnx

diff --git a/Engine/Plugins/Animation/ControlRigSpline/Source/ControlRigSpline/ControlRigSpline.build.cs b/Engine/Plugins/Animation/ControlRigSpline/Source/ControlRigSpline/ControlRigSpline.build.cs
--- a/Engine/Plugins/Animation/ControlRigSpline/Source/ControlRigSpline/ControlRigSpline.build.cs
+++ b/Engine/Plugins/Animation/ControlRigSpline/Source/ControlRigSpline/ControlRigSpline.build.cs
@@ -1,12 +1,14 @@
 // Copyright Epic Games, Inc. All Rights Reserved.
 
+using System.IO;
+
 namespace UnrealBuildTool.Rules
 {
 	public class ControlRigSpline : ModuleRules
 	{
 		public ControlRigSpline(ReadOnlyTargetRules Target) : base(Target)
 		{
-			PrivateIncludePaths.Add("ControlRigSpline/ThirdParty/TinySpline");
+			PrivateIncludePaths.Add(Path.Combine(ModuleDirectory, "ThirdParty/TinySpline"));
 
 			PublicDependencyModuleNames.AddRange(
 			new string[]
@@ -17,6 +19,9 @@ namespace UnrealBuildTool.Rules
 				"RigVM",
 				"ControlRig",
 			});
+
+			// TODO: Should not be including private headers in public code
+			PublicIncludePaths.Add(Path.Combine(GetModuleDirectory("ControlRig"), "Private/Units/Highlevel/Hierarchy")); // For RigUnit_FitChainToCurve.h
 		}
 	}
 }
diff --git a/Engine/Plugins/Animation/ControlRigSpline/Source/ControlRigSpline/Private/ControlRigSplineUnits.cpp b/Engine/Plugins/Animation/ControlRigSpline/Source/ControlRigSpline/Private/ControlRigSplineUnits.cpp
--- a/Engine/Plugins/Animation/ControlRigSpline/Source/ControlRigSpline/Private/ControlRigSplineUnits.cpp
+++ b/Engine/Plugins/Animation/ControlRigSpline/Source/ControlRigSpline/Private/ControlRigSplineUnits.cpp
@@ -4,6 +4,7 @@
 #include "Units/RigUnitContext.h"
 #include "Features/IModularFeatures.h"
 #include "Algo/BinarySearch.h"
+#include "RigUnit_FitChainToCurve.h"
 
 #include "tinysplinecxx.h"
 
diff --git a/Engine/Plugins/Animation/ControlRigSpline/Source/ControlRigSpline/Public/ControlRigSplineUnits.h b/Engine/Plugins/Animation/ControlRigSpline/Source/ControlRigSpline/Public/ControlRigSplineUnits.h
--- a/Engine/Plugins/Animation/ControlRigSpline/Source/ControlRigSpline/Public/ControlRigSplineUnits.h
+++ b/Engine/Plugins/Animation/ControlRigSpline/Source/ControlRigSpline/Public/ControlRigSplineUnits.h
@@ -5,7 +5,7 @@
 #include "Units/RigUnit.h"
 #include "ControlRig.h"
 #include "ControlRigSplineTypes.h"
-#include "ControlRig/Private/Units/Highlevel/Hierarchy/RigUnit_FitChainToCurve.h"
+#include "RigUnit_FitChainToCurve.h"
 #include "ControlRigSplineUnits.generated.h"
 
 USTRUCT(meta = (Abstract, NodeColor = "0.3 0.1 0.1"))
