Control Rig Spline: Closed curves

diff --git a/Engine/Plugins/Animation/ControlRigSpline/Source/ControlRigSpline/Private/ControlRigSplineTypes.cpp b/Engine/Plugins/Animation/ControlRigSpline/Source/ControlRigSpline/Private/ControlRigSplineTypes.cpp
--- a/Engine/Plugins/Animation/ControlRigSpline/Source/ControlRigSpline/Private/ControlRigSplineTypes.cpp
+++ b/Engine/Plugins/Animation/ControlRigSpline/Source/ControlRigSpline/Private/ControlRigSplineTypes.cpp
@@ -4,15 +4,48 @@
 
 #include UE_INLINE_GENERATED_CPP_BY_NAME(ControlRigSplineTypes)
 
-ControlRigBaseSpline::ControlRigBaseSpline(const TArrayView<const FVector>& InControlPoints, const uint16 InDegree)
-: ControlPoints(InControlPoints)
-, Degree(InDegree)
+ControlRigBaseSpline::ControlRigBaseSpline(const TArrayView<const FVector>& InControlPoints, const uint16 InDegree, bool bInClosed)
+: Degree(InDegree)
+, bClosed(bInClosed)
 {
-	
+	if (bClosed)
+	{
+		int32 NumPoints = InControlPoints.Num() + Degree;
+		ControlPoints.Reserve(NumPoints);
+		
+		int32 Index = InControlPoints.Num() - Degree/2;
+		while (ControlPoints.Num() < NumPoints)
+		{
+			Index = (Index+1)%InControlPoints.Num();
+			ControlPoints.Add(InControlPoints[Index]);
+		}
+	}
+	else
+	{
+		ControlPoints = InControlPoints;
+	}
+}
+
+void ControlRigBaseSpline::SetControlPoints(const TArrayView<const FVector>& InControlPoints)
+{
+	if (bClosed)
+	{
+		int32 CurIndex = 0;
+		int32 Index = InControlPoints.Num() - Degree/2;
+		while (CurIndex < ControlPoints.Num())
+		{
+			Index = (Index+1)%InControlPoints.Num();
+			ControlPoints[CurIndex++] = InControlPoints[Index];
+		}
+	}
+	else
+	{
+		ControlPoints = InControlPoints;
+	}
 }
 
-ControlRigBSpline::ControlRigBSpline(const TArrayView<const FVector>& InControlPoints, const uint16 InDegree, bool bInClamped)
-: ControlRigBaseSpline(InControlPoints, InDegree)
+ControlRigBSpline::ControlRigBSpline(const TArrayView<const FVector>& InControlPoints, const uint16 InDegree, const bool bInClosed, bool bInClamped)
+: ControlRigBaseSpline(InControlPoints, InDegree, bInClosed)
 {
 	uint16 N = ControlPoints.Num() - 1;
 	uint16 P = Degree;
@@ -20,7 +53,7 @@ ControlRigBSpline::ControlRigBSpline(const TArrayView<const FVector>& InControlP
 	KnotVector.SetNumZeroed(M+1);
 
 	// Clamped means there is p+1 0s at the beginning and 1s at the end of the knot vector
-	if (bInClamped)
+	if (bInClamped && !bInClosed)
 	{
 		for (int32 i=M-P; i<=M; ++i)
 		{
@@ -101,13 +134,20 @@ FVector ControlRigBSpline::GetPointAtParam(float Param)
 	return WorkingPoints[0];
 }
 
-ControlRigHermite::ControlRigHermite(const TArrayView<const FVector>& InControlPoints)
-: ControlRigBaseSpline(InControlPoints, 3)
+ControlRigHermite::ControlRigHermite(const TArrayView<const FVector>& InControlPoints, const bool bInClosed)
+: ControlRigBaseSpline(InControlPoints, 4, bInClosed)
 {
-	SegmentPoints.Add(ControlPoints[0] + ControlPoints[0] - ControlPoints[1]);
-	SegmentPoints.Append(ControlPoints);
-	SegmentPoints.Add(ControlPoints.Last() + ControlPoints.Last() - ControlPoints[ControlPoints.Num()-2]);
-
+	if (bClosed)
+	{
+		SegmentPoints.Append(ControlPoints);
+	}
+	else
+	{		
+		SegmentPoints.Add(ControlPoints[0] + ControlPoints[0] - ControlPoints[1]);
+		SegmentPoints.Append(ControlPoints);
+		SegmentPoints.Add(ControlPoints.Last() + ControlPoints.Last() - ControlPoints[ControlPoints.Num()-2]);
+	}
+	
 	NumSegments = SegmentPoints.Num() - 3;
 }
 
@@ -116,9 +156,16 @@ void ControlRigHermite::SetControlPoints(const TArrayView<const FVector>& InCont
 	ControlRigBaseSpline::SetControlPoints(InControlPoints);
 
 	SegmentPoints.Reset();
-	SegmentPoints.Add(ControlPoints[0] + ControlPoints[0] - ControlPoints[1]);
-	SegmentPoints.Append(ControlPoints);
-	SegmentPoints.Add(ControlPoints.Last() + ControlPoints.Last() - ControlPoints[ControlPoints.Num()-2]);
+	if (bClosed)
+	{
+		SegmentPoints.Append(ControlPoints);
+	}
+	else
+	{
+		SegmentPoints.Add(ControlPoints[0] + ControlPoints[0] - ControlPoints[1]);
+		SegmentPoints.Append(ControlPoints);
+		SegmentPoints.Add(ControlPoints.Last() + ControlPoints.Last() - ControlPoints[ControlPoints.Num()-2]);
+	}	
 }
 
 static FVector Hermite(const FVector& P0, const FVector& P1, const FVector& M0, const FVector& M1, float Param)
@@ -184,6 +231,12 @@ TArray<FVector>& FControlRigSplineImpl::GetControlPoints()
 	return Spline->GetControlPoints();	
 }
 
+uint8 FControlRigSplineImpl::GetDegree() const
+{
+	check(Spline);
+	return Spline->GetDegree();
+}
+
 FControlRigSpline::FControlRigSpline(const FControlRigSpline& InOther)
 {
 	SplineData = InOther.SplineData;
@@ -195,7 +248,16 @@ FControlRigSpline& FControlRigSpline::operator=(const FControlRigSpline& InOther
 	return *this;
 }
 
-void FControlRigSpline::SetControlPoints(const TArrayView<const FVector>& InPoints, const ESplineType SplineMode, const int32 SamplesPerSegment, const float Compression, const float Stretch)
+uint8 FControlRigSpline::GetDegree() const
+{
+	if (!SplineData.IsValid())
+	{
+		return 0;
+	}
+	return SplineData->GetDegree();
+}
+
+void FControlRigSpline::SetControlPoints(const TArrayView<const FVector>& InPoints, const ESplineType SplineMode, const bool bInClosed, const int32 SamplesPerSegment, const float Compression, const float Stretch)
 {
 	const int32 ControlPointsCount = InPoints.Num();
 	if (ControlPointsCount < 4)
@@ -213,20 +275,22 @@ void FControlRigSpline::SetControlPoints(const TArrayView<const FVector>& InPoin
 	bool bSplineModeChanged = SplineMode != SplineData->SplineMode;
 	bool bSamplesCountChanged = SamplesPerSegment != SplineData->SamplesPerSegment;
 	bool bStretchChanged = Stretch != SplineData->Stretch || Compression != SplineData->Compression;
-	if (!bSplineModeChanged && !bControlPointsChanged && !bSamplesCountChanged && !bStretchChanged)
+	bool bClosedChanged = bInClosed != SplineData->bClosed;
+	if (!bSplineModeChanged && !bControlPointsChanged && !bSamplesCountChanged && !bStretchChanged && !bClosedChanged)
 	{
 		return;
 	}
 
 	SplineData->SplineMode = SplineMode;
+	SplineData->bClosed = bInClosed;
 	SplineData->SamplesPerSegment = SamplesPerSegment;
 	SplineData->Compression = Compression;
 	SplineData->Stretch = Stretch;
 
 	// If we need to update the spline because the controls points have changed, or the spline mode has changed
-	if (bControlPointsChanged || bSplineModeChanged)
+	if (bControlPointsChanged || bSplineModeChanged || bClosedChanged)
 	{	
-		if (bNumControlPointsChanged || bSplineModeChanged)
+		if (bNumControlPointsChanged || bSplineModeChanged || bClosedChanged)
 		{
 			// Delete previously created spline
 			if (SplineData->Spline)
@@ -239,13 +303,13 @@ void FControlRigSpline::SetControlPoints(const TArrayView<const FVector>& InPoin
 			{
 				case ESplineType::BSpline:
 				{
-					SplineData->Spline = new ControlRigBSpline(InPoints, 3, true);
+					SplineData->Spline = new ControlRigBSpline(InPoints, 3, bInClosed, true);
 					break;
 				}
 				case ESplineType::Hermite:
 				{
 					
-					SplineData->Spline = new ControlRigHermite(InPoints);
+					SplineData->Spline = new ControlRigHermite(InPoints, bInClosed);
 					break;
 				}
 				default:
@@ -262,7 +326,7 @@ void FControlRigSpline::SetControlPoints(const TArrayView<const FVector>& InPoin
 	}
 
 	// If curve has changed, or sample count has changed, recompute the cache
-	if (bControlPointsChanged || bSplineModeChanged || bSamplesCountChanged || bStretchChanged)
+	if (bControlPointsChanged || bSplineModeChanged || bSamplesCountChanged || bStretchChanged || bClosedChanged)
 	{
 		// Cache sample positions of the spline
 		int32 NumSamples = (ControlPointsCount-1) * SamplesPerSegment;
@@ -276,7 +340,7 @@ void FControlRigSpline::SetControlPoints(const TArrayView<const FVector>& InPoin
 		}
 
 		// Correct length of samples
-		if (!bSplineModeChanged && !bSamplesCountChanged && !bNumControlPointsChanged)
+		if (!bSplineModeChanged && !bSamplesCountChanged && !bNumControlPointsChanged && !bClosedChanged)
 		{
 			if (SplineData->InitialLengths.Num() > 0)
 			{
@@ -306,14 +370,14 @@ void FControlRigSpline::SetControlPoints(const TArrayView<const FVector>& InPoin
 		{
 			SplineData->AccumulatedLenth.SetNumUninitialized(SplineData->SamplesArray.Num(), false);
 			SplineData->AccumulatedLenth[0] = 0.f;
-			if (bSplineModeChanged || bSamplesCountChanged || bNumControlPointsChanged)
+			if (bSplineModeChanged || bClosedChanged || bSamplesCountChanged || bNumControlPointsChanged)
 			{
 				SplineData->InitialLengths.SetNumUninitialized(SplineData->SamplesArray.Num() - 1);
 			}
 			for (int32 i = 1; i < SplineData->SamplesArray.Num(); ++i)
 			{
 				float CurrentLength = FVector::Distance(SplineData->SamplesArray[i - 1], SplineData->SamplesArray[i]);
-				if (bSplineModeChanged || bSamplesCountChanged || bNumControlPointsChanged)
+				if (bSplineModeChanged || bClosedChanged || bSamplesCountChanged || bNumControlPointsChanged)
 				{
 					SplineData->InitialLengths[i-1] = CurrentLength;
 				}
diff --git a/Engine/Plugins/Animation/ControlRigSpline/Source/ControlRigSpline/Private/ControlRigSplineUnits.cpp b/Engine/Plugins/Animation/ControlRigSpline/Source/ControlRigSpline/Private/ControlRigSplineUnits.cpp
--- a/Engine/Plugins/Animation/ControlRigSpline/Source/ControlRigSpline/Private/ControlRigSplineUnits.cpp
+++ b/Engine/Plugins/Animation/ControlRigSpline/Source/ControlRigSpline/Private/ControlRigSplineUnits.cpp
@@ -20,7 +20,7 @@ FRigUnit_ControlRigSplineFromPoints_Execute()
 	}
 		
 	const TArrayView<const FVector> PointsView(Points.GetData(), Points.Num());
-	Spline.SetControlPoints(PointsView, SplineMode, SamplesPerSegment, Compression, Stretch);
+	Spline.SetControlPoints(PointsView, SplineMode, bClosed, SamplesPerSegment, Compression, Stretch);
 }
 
 FRigUnit_SetSplinePoints_Execute()
@@ -37,14 +37,25 @@ FRigUnit_SetSplinePoints_Execute()
 		return;
 	}	
 
-	if (Points.Num() != Spline.SplineData->GetControlPoints().Num())
+	if (!Spline.SplineData->bClosed)
 	{
-		UE_LOG(LogControlRig, Error, TEXT("Number of input points does not match the number of point in the spline."));
-		return;
+		if (Points.Num() != Spline.SplineData->GetControlPoints().Num())
+		{
+			UE_LOG(LogControlRig, Error, TEXT("Number of input points does not match the number of point in the spline."));
+			return;
+		}
+	}
+	else if(Points.Num() != Spline.SplineData->GetControlPoints().Num() - Spline.GetDegree())
+	{
+		if (Points.Num() != Spline.SplineData->GetControlPoints().Num())
+		{
+			UE_LOG(LogControlRig, Error, TEXT("Number of input points does not match the number of point in the spline."));
+			return;
+		}
 	}
 
 	const TArrayView<const FVector> PointsView(Points.GetData(), Points.Num());
-	Spline.SetControlPoints(PointsView, Spline.SplineData->SplineMode, Spline.SplineData->SamplesPerSegment, Spline.SplineData->Compression, Spline.SplineData->Stretch);
+	Spline.SetControlPoints(PointsView, Spline.SplineData->SplineMode, Spline.SplineData->bClosed, Spline.SplineData->SamplesPerSegment, Spline.SplineData->Compression, Spline.SplineData->Stretch);
 }
 
 FRigUnit_PositionFromControlRigSpline_Execute()
@@ -679,7 +690,7 @@ FRigUnit_FitSplineCurveToChainItemArray_Execute()
 	}
     FControlRigSpline AuxSpline;
 	const TArrayView<const FVector> PointsView(AuxControlPoints.GetData(), AuxControlPoints.Num());
-	AuxSpline.SetControlPoints(PointsView, ESplineType::Hermite, Spline.SplineData->SamplesPerSegment);
+	AuxSpline.SetControlPoints(PointsView, ESplineType::Hermite, false, Spline.SplineData->SamplesPerSegment);
 	
 
 	// 2.-  for each control point in the original spline
@@ -696,7 +707,7 @@ FRigUnit_FitSplineCurveToChainItemArray_Execute()
 	}
 
 	const TArrayView<const FVector> NewPointsView(NewControlPoints.GetData(), NewControlPoints.Num());
-	Spline.SetControlPoints(NewPointsView, Spline.SplineData->SplineMode, Spline.SplineData->SamplesPerSegment, Spline.SplineData->Compression, Spline.SplineData->Stretch);
+	Spline.SetControlPoints(NewPointsView, Spline.SplineData->SplineMode, false, Spline.SplineData->SamplesPerSegment, Spline.SplineData->Compression, Spline.SplineData->Stretch);
 }
 
 FRigUnit_ClosestParameterFromControlRigSpline_Execute()
diff --git a/Engine/Plugins/Animation/ControlRigSpline/Source/ControlRigSpline/Public/ControlRigSplineTypes.h b/Engine/Plugins/Animation/ControlRigSpline/Source/ControlRigSpline/Public/ControlRigSplineTypes.h
--- a/Engine/Plugins/Animation/ControlRigSpline/Source/ControlRigSpline/Public/ControlRigSplineTypes.h
+++ b/Engine/Plugins/Animation/ControlRigSpline/Source/ControlRigSpline/Public/ControlRigSplineTypes.h
@@ -39,17 +39,20 @@ class ControlRigBaseSpline
 protected:
 	TArray<FVector> ControlPoints;
 	uint16 Degree;
+	bool bClosed;
 	
 public:
 
-	ControlRigBaseSpline(const TArrayView<const FVector>& InControlPoints, uint16 InDegree);
+	ControlRigBaseSpline(const TArrayView<const FVector>& InControlPoints, uint16 InDegree, bool bInClosed);
 	virtual ~ControlRigBaseSpline(){}
 	
 	virtual FVector GetPointAtParam(float Param) = 0;
 
-	virtual void SetControlPoints(const TArrayView<const FVector>& InControlPoints) { ControlPoints = InControlPoints; }
+	virtual void SetControlPoints(const TArrayView<const FVector>& InControlPoints);
 
 	TArray<FVector>& GetControlPoints() { return ControlPoints; }
+
+	uint8 GetDegree() const { return Degree; }
 };
 
 class ControlRigBSpline : public ControlRigBaseSpline
@@ -57,8 +60,8 @@ class ControlRigBSpline : public ControlRigBaseSpline
 	TArray<float> KnotVector;
 	
 public:
-	ControlRigBSpline(const TArrayView<const FVector>& InControlPoints, const uint16 InDegree, const bool bInClamped);
-	
+	ControlRigBSpline(const TArrayView<const FVector>& InControlPoints, const uint16 InDegree, const bool bInClosed, const bool bInClamped);
+
 	virtual FVector GetPointAtParam(float Param);
 
 };
@@ -70,7 +73,7 @@ class ControlRigHermite : public ControlRigBaseSpline
 	uint16 NumSegments;
 	
 public:
-	ControlRigHermite(const TArrayView<const FVector>& InControlPoints);
+	ControlRigHermite(const TArrayView<const FVector>& InControlPoints, const bool bInClosed);
 
 	virtual void SetControlPoints(const TArrayView<const FVector>& InControlPoints) override;
 	
@@ -88,6 +91,7 @@ struct CONTROLRIGSPLINE_API FControlRigSplineImpl
 	{
 		SplineMode = ESplineType::BSpline;
 		Spline = nullptr;
+		bClosed = false;
 		SamplesPerSegment = 16;
 	}
 
@@ -102,6 +106,9 @@ struct CONTROLRIGSPLINE_API FControlRigSplineImpl
 	// The actual spline
 	ControlRigBaseSpline* Spline;
 
+	// Wether or not the curve is closed
+	bool bClosed;
+
 	// Samples per segment, where segment is the portion between two control points
 	int32 SamplesPerSegment;
 
@@ -119,6 +126,9 @@ struct CONTROLRIGSPLINE_API FControlRigSplineImpl
 
 	// Returns a reference to the control points that were used to create this spline 
 	TArray<FVector>& GetControlPoints();
+
+	// Returns the degree of the curve
+	uint8 GetDegree() const;
 };
 
 USTRUCT(BlueprintType)
@@ -135,6 +145,13 @@ struct CONTROLRIGSPLINE_API FControlRigSpline
 
 	TSharedPtr<FControlRigSplineImpl> SplineData;
 
+	/**
+	* Returns the degree of the spline
+	*
+	* @return			The degree of the spline.
+	*/
+	uint8 GetDegree() const;
+
 	/**
 	* Sets the control points in the spline. It will build the spline if needed, or forceRebuild is true,
 	* or will update the points if building from scratch is not necessary. The type of spline to build will
@@ -142,11 +159,12 @@ struct CONTROLRIGSPLINE_API FControlRigSpline
 	*
 	* @param InPoints	The control points to set.
 	* @param SplineMode	The type of spline
+	* @param bInClosed	If the spline should be closed
 	* @param SamplesPerSegment The samples to cache for every segment defined between two control rig
 	* @param Compression The allowed length compression (1.f being do not allow compression`). If 0, no restriction wil be applied.
 	* @param Stretch The allowed length stretch (1.f being do not allow stretch). If 0, no restriction wil be applied.
 	*/
-	void SetControlPoints(const TArrayView<const FVector>& InPoints, const ESplineType SplineMode = ESplineType::BSpline, const int32 SamplesPerSegment = 16, const float Compression = 1.f, const float Stretch = 1.f);
+	void SetControlPoints(const TArrayView<const FVector>& InPoints, const ESplineType SplineMode = ESplineType::BSpline, const bool bInClosed = false, const int32 SamplesPerSegment = 16, const float Compression = 1.f, const float Stretch = 1.f);
 
 	/**
 	* Given an InParam float in [0, 1], will return the position of the spline at that point.
diff --git a/Engine/Plugins/Animation/ControlRigSpline/Source/ControlRigSpline/Public/ControlRigSplineUnits.h b/Engine/Plugins/Animation/ControlRigSpline/Source/ControlRigSpline/Public/ControlRigSplineUnits.h
--- a/Engine/Plugins/Animation/ControlRigSpline/Source/ControlRigSpline/Public/ControlRigSplineUnits.h
+++ b/Engine/Plugins/Animation/ControlRigSpline/Source/ControlRigSpline/Public/ControlRigSplineUnits.h
@@ -25,6 +25,7 @@ struct CONTROLRIGSPLINE_API FRigUnit_ControlRigSplineFromPoints : public FRigUni
 	FRigUnit_ControlRigSplineFromPoints()
 	{
 		SplineMode = ESplineType::Hermite;
+		bClosed = false;
 		SamplesPerSegment = 16;
 		Compression = 0.f;
 		Stretch = 0.f;
@@ -40,6 +41,9 @@ struct CONTROLRIGSPLINE_API FRigUnit_ControlRigSplineFromPoints : public FRigUni
 	UPROPERTY(meta = (Input))
 	ESplineType SplineMode;
 
+	UPROPERTY(meta = (Input))
+	bool bClosed;
+
 	UPROPERTY(meta = (Input))
 	int32 SamplesPerSegment;
 
