Control Rig Spline: Fix closed hermite splines do not have the same value at U=0 and U=1

diff --git a/Engine/Plugins/Animation/ControlRigSpline/Source/ControlRigSpline/Private/ControlRigSplineTypes.cpp b/Engine/Plugins/Animation/ControlRigSpline/Source/ControlRigSpline/Private/ControlRigSplineTypes.cpp
--- a/Engine/Plugins/Animation/ControlRigSpline/Source/ControlRigSpline/Private/ControlRigSplineTypes.cpp
+++ b/Engine/Plugins/Animation/ControlRigSpline/Source/ControlRigSpline/Private/ControlRigSplineTypes.cpp
@@ -140,15 +140,15 @@ ControlRigHermite::ControlRigHermite(const TArrayView<const FVector>& InControlP
 	if (bClosed)
 	{
 		SegmentPoints.Append(ControlPoints);
+		NumSegments = SegmentPoints.Num() - 4;
 	}
 	else
-	{		
+	{
 		SegmentPoints.Add(ControlPoints[0] + ControlPoints[0] - ControlPoints[1]);
 		SegmentPoints.Append(ControlPoints);
 		SegmentPoints.Add(ControlPoints.Last() + ControlPoints.Last() - ControlPoints[ControlPoints.Num()-2]);
+		NumSegments = SegmentPoints.Num() - 3;
 	}
-	
-	NumSegments = SegmentPoints.Num() - 3;
 }
 
 void ControlRigHermite::SetControlPoints(const TArrayView<const FVector>& InControlPoints)
@@ -165,7 +165,7 @@ void ControlRigHermite::SetControlPoints(const TArrayView<const FVector>& InCont
 		SegmentPoints.Add(ControlPoints[0] + ControlPoints[0] - ControlPoints[1]);
 		SegmentPoints.Append(ControlPoints);
 		SegmentPoints.Add(ControlPoints.Last() + ControlPoints.Last() - ControlPoints[ControlPoints.Num()-2]);
-	}	
+	}
 }
 
 static FVector Hermite(const FVector& P0, const FVector& P1, const FVector& M0, const FVector& M1, float Param)
@@ -349,18 +349,25 @@ void FControlRigSpline::SetControlPoints(const TArrayView<const FVector>& InPoin
 				{
 					for (int32 j = (i == 0) ? 1 : 0; j < SamplesPerSegment; ++j)
 					{
+						const int32 CurSampleIndex = i * SamplesPerSegment + j;
+						const int32 PrevSampleIndex = CurSampleIndex - 1;
+						if (!SplineData->InitialLengths.IsValidIndex(PrevSampleIndex))
+						{
+							break;
+						}
+						
 						// Get direction from samples before correction
-						FVector Dir = SamplesBeforeCorrect[i * SamplesPerSegment + j] - SamplesBeforeCorrect[i * SamplesPerSegment + j - 1];
+						FVector Dir = SamplesBeforeCorrect[CurSampleIndex] - SamplesBeforeCorrect[PrevSampleIndex];
 						Dir.Normalize();
 
-						float InitialLength = SplineData->InitialLengths[i * SamplesPerSegment + j - 1];
+						float InitialLength = SplineData->InitialLengths[PrevSampleIndex];
 						// Current length as the projection on the Dir vector (might be negative)
-						float CurrentLength = (SplineData->SamplesArray[i * SamplesPerSegment + j] - SplineData->SamplesArray[i * SamplesPerSegment + j - 1]).Dot(Dir);
+						float CurrentLength = (SplineData->SamplesArray[CurSampleIndex] - SplineData->SamplesArray[PrevSampleIndex]).Dot(Dir);
 						float FixedLength = FMath::Clamp(CurrentLength,
 							(Compression > 0.f) ? InitialLength * Compression : CurrentLength,
 							(Stretch > 0.f) ? InitialLength * Stretch : CurrentLength);
 
-						SplineData->SamplesArray[i * SamplesPerSegment + j] = SplineData->SamplesArray[i * SamplesPerSegment + j - 1] + Dir * FixedLength;
+						SplineData->SamplesArray[CurSampleIndex] = SplineData->SamplesArray[PrevSampleIndex] + Dir * FixedLength;
 					}
 				}
 			}
