Control Rig Spline: Fix backwards solve for closed splines

diff --git a/Engine/Plugins/Animation/ControlRigSpline/Source/ControlRigSpline/Private/ControlRigSplineTypes.cpp b/Engine/Plugins/Animation/ControlRigSpline/Source/ControlRigSpline/Private/ControlRigSplineTypes.cpp
--- a/Engine/Plugins/Animation/ControlRigSpline/Source/ControlRigSpline/Private/ControlRigSplineTypes.cpp
+++ b/Engine/Plugins/Animation/ControlRigSpline/Source/ControlRigSpline/Private/ControlRigSplineTypes.cpp
@@ -26,6 +26,25 @@ ControlRigBaseSpline::ControlRigBaseSpline(const TArrayView<const FVector>& InCo
 	}
 }
 
+TArray<FVector> ControlRigBaseSpline::GetControlPointsWithoutDuplicates()
+{
+	if (bClosed)
+	{
+		const int32 NumPoints = ControlPoints.Num() - Degree;
+		const int32 NumInitPointsToIgnore = Degree/2 - 1;
+		
+		TArray<FVector> ReducedControlPoints;
+		ReducedControlPoints.Reserve(NumPoints);
+		for (int32 i=0; i<NumPoints; ++i)
+		{
+			ReducedControlPoints.Add(ControlPoints[i+NumInitPointsToIgnore]);
+		}
+		return ReducedControlPoints;
+	}
+
+	return ControlPoints;
+}
+
 void ControlRigBaseSpline::SetControlPoints(const TArrayView<const FVector>& InControlPoints)
 {
 	if (bClosed)
@@ -231,6 +250,12 @@ TArray<FVector>& FControlRigSplineImpl::GetControlPoints()
 	return Spline->GetControlPoints();	
 }
 
+TArray<FVector> FControlRigSplineImpl::GetControlPointsWithoutDuplicates()
+{
+	check(Spline);
+	return Spline->GetControlPointsWithoutDuplicates();
+}
+
 uint8 FControlRigSplineImpl::GetDegree() const
 {
 	check(Spline);
diff --git a/Engine/Plugins/Animation/ControlRigSpline/Source/ControlRigSpline/Private/ControlRigSplineUnits.cpp b/Engine/Plugins/Animation/ControlRigSpline/Source/ControlRigSpline/Private/ControlRigSplineUnits.cpp
--- a/Engine/Plugins/Animation/ControlRigSpline/Source/ControlRigSpline/Private/ControlRigSplineUnits.cpp
+++ b/Engine/Plugins/Animation/ControlRigSpline/Source/ControlRigSpline/Private/ControlRigSplineUnits.cpp
@@ -680,6 +680,8 @@ FRigUnit_FitSplineCurveToChainItemArray_Execute()
 		return;
 	}
 
+	const bool bClosed = Spline.SplineData->bClosed;
+
 	// 1.- Create spline from chain
 	TArray<FVector> AuxControlPoints;
 	AuxControlPoints.SetNumUninitialized(Items.Num());	
@@ -690,24 +692,27 @@ FRigUnit_FitSplineCurveToChainItemArray_Execute()
 	}
     FControlRigSpline AuxSpline;
 	const TArrayView<const FVector> PointsView(AuxControlPoints.GetData(), AuxControlPoints.Num());
-	AuxSpline.SetControlPoints(PointsView, ESplineType::Hermite, false, Spline.SplineData->SamplesPerSegment);
+	AuxSpline.SetControlPoints(PointsView, ESplineType::Hermite, bClosed, Spline.SplineData->SamplesPerSegment);
 	
 
 	// 2.-  for each control point in the original spline
 	//			figure out its u in the original spline (preprocess)
 	//			query position at u on the new spline
-	const TArray<FVector>& ControlPoints = Spline.SplineData->GetControlPoints();
+	const TArray<FVector> ControlPoints = Spline.SplineData->GetControlPointsWithoutDuplicates();
+	const int32 NumControlPoints = ControlPoints.Num();
 	TArray<FVector> NewControlPoints;
-	NewControlPoints.SetNumUninitialized(ControlPoints.Num());
-	for (int32 i = 0; i < ControlPoints.Num(); ++i)
+	NewControlPoints.SetNumUninitialized(NumControlPoints);
+	float U = 0;
+	const float DeltaU = bClosed ? 1 / (float)(NumControlPoints) : 1 / (float)(NumControlPoints - 1);
+	for (int32 i = 0; i < NumControlPoints; ++i)
 	{
-		float U = i / (float)(ControlPoints.Num() - 1);
 		FVector NewPosition = AuxSpline.PositionAtParam(U);
 		NewControlPoints[i] = NewPosition;
+		U += DeltaU;
 	}
 
 	const TArrayView<const FVector> NewPointsView(NewControlPoints.GetData(), NewControlPoints.Num());
-	Spline.SetControlPoints(NewPointsView, Spline.SplineData->SplineMode, false, Spline.SplineData->SamplesPerSegment, Spline.SplineData->Compression, Spline.SplineData->Stretch);
+	Spline.SetControlPoints(NewPointsView, Spline.SplineData->SplineMode, bClosed, Spline.SplineData->SamplesPerSegment, Spline.SplineData->Compression, Spline.SplineData->Stretch);
 }
 
 FRigUnit_ClosestParameterFromControlRigSpline_Execute()
diff --git a/Engine/Plugins/Animation/ControlRigSpline/Source/ControlRigSpline/Public/ControlRigSplineTypes.h b/Engine/Plugins/Animation/ControlRigSpline/Source/ControlRigSpline/Public/ControlRigSplineTypes.h
--- a/Engine/Plugins/Animation/ControlRigSpline/Source/ControlRigSpline/Public/ControlRigSplineTypes.h
+++ b/Engine/Plugins/Animation/ControlRigSpline/Source/ControlRigSpline/Public/ControlRigSplineTypes.h
@@ -50,7 +50,9 @@ public:
 
 	virtual void SetControlPoints(const TArrayView<const FVector>& InControlPoints);
 
-	TArray<FVector>& GetControlPoints() { return ControlPoints; }
+	virtual TArray<FVector>& GetControlPoints() { return ControlPoints; }
+	
+	virtual TArray<FVector> GetControlPointsWithoutDuplicates();
 
 	uint8 GetDegree() const { return Degree; }
 };
@@ -127,6 +129,9 @@ struct CONTROLRIGSPLINE_API FControlRigSplineImpl
 	// Returns a reference to the control points that were used to create this spline 
 	TArray<FVector>& GetControlPoints();
 
+	// Returns the control points that were used to create this spline, removing the duplicates in case of a closed spline
+	TArray<FVector> GetControlPointsWithoutDuplicates();
+
 	// Returns the degree of the curve
 	uint8 GetDegree() const;
 };
