Control Rig Spline: Fix closed BSpline not starting closer to the first control point

diff --git a/Engine/Plugins/Animation/ControlRigSpline/Source/ControlRigSpline/Private/ControlRigSplineTypes.cpp b/Engine/Plugins/Animation/ControlRigSpline/Source/ControlRigSpline/Private/ControlRigSplineTypes.cpp
--- a/Engine/Plugins/Animation/ControlRigSpline/Source/ControlRigSpline/Private/ControlRigSplineTypes.cpp
+++ b/Engine/Plugins/Animation/ControlRigSpline/Source/ControlRigSpline/Private/ControlRigSplineTypes.cpp
@@ -13,7 +13,7 @@ ControlRigBaseSpline::ControlRigBaseSpline(const TArrayView<const FVector>& InCo
 		int32 NumPoints = InControlPoints.Num() + Degree;
 		ControlPoints.Reserve(NumPoints);
 		
-		int32 Index = InControlPoints.Num() - Degree/2;
+		int32 Index = InControlPoints.Num() - FMath::RoundToPositiveInfinity((float)Degree/2);
 		while (ControlPoints.Num() < NumPoints)
 		{
 			Index = (Index+1)%InControlPoints.Num();
@@ -31,7 +31,7 @@ TArray<FVector> ControlRigBaseSpline::GetControlPointsWithoutDuplicates()
 	if (bClosed)
 	{
 		const int32 NumPoints = ControlPoints.Num() - Degree;
-		const int32 NumInitPointsToIgnore = Degree/2 - 1;
+		const int32 NumInitPointsToIgnore = FMath::RoundToPositiveInfinity((float)Degree/2) - 1;
 		
 		TArray<FVector> ReducedControlPoints;
 		ReducedControlPoints.Reserve(NumPoints);
@@ -50,7 +50,7 @@ void ControlRigBaseSpline::SetControlPoints(const TArrayView<const FVector>& InC
 	if (bClosed)
 	{
 		int32 CurIndex = 0;
-		int32 Index = InControlPoints.Num() - Degree/2;
+		int32 Index = InControlPoints.Num() - FMath::RoundToPositiveInfinity((float)Degree/2);
 		while (CurIndex < ControlPoints.Num())
 		{
 			Index = (Index+1)%InControlPoints.Num();
@@ -295,8 +295,14 @@ void FControlRigSpline::SetControlPoints(const TArrayView<const FVector>& InPoin
 		SplineData = MakeShared<FControlRigSplineImpl>();
 	}
 
-	bool bControlPointsChanged = (SplineData->Spline) ? InPoints != SplineData->GetControlPoints() : true;
-	bool bNumControlPointsChanged = (SplineData->Spline) ? SplineData->GetControlPoints().Num() != ControlPointsCount : true;
+	TArray<FVector> OldControlPoints;
+	if (SplineData->Spline)
+	{
+		OldControlPoints = SplineData->GetControlPointsWithoutDuplicates();
+	}
+
+	bool bControlPointsChanged = (SplineData->Spline) ? InPoints != OldControlPoints : true;
+	bool bNumControlPointsChanged = (SplineData->Spline) ? OldControlPoints.Num() != ControlPointsCount : true;
 	bool bSplineModeChanged = SplineMode != SplineData->SplineMode;
 	bool bSamplesCountChanged = SamplesPerSegment != SplineData->SamplesPerSegment;
 	bool bStretchChanged = Stretch != SplineData->Stretch || Compression != SplineData->Compression;
