Control Rig: Fix crash in ControlRigSpline when samples per segment is set to 0

diff --git a/Engine/Plugins/Animation/ControlRigSpline/Source/ControlRigSpline/Private/ControlRigSplineTypes.cpp b/Engine/Plugins/Animation/ControlRigSpline/Source/ControlRigSpline/Private/ControlRigSplineTypes.cpp
--- a/Engine/Plugins/Animation/ControlRigSpline/Source/ControlRigSpline/Private/ControlRigSplineTypes.cpp
+++ b/Engine/Plugins/Animation/ControlRigSpline/Source/ControlRigSpline/Private/ControlRigSplineTypes.cpp
@@ -290,6 +290,11 @@ void FControlRigSpline::SetControlPoints(const TArrayView<const FVector>& InPoin
 		return;
 	}
 
+	if (SamplesPerSegment < 1)
+	{
+		return;
+	}
+
 	if (!SplineData.IsValid())
 	{
 		SplineData = MakeShared<FControlRigSplineImpl>();
@@ -405,6 +410,7 @@ void FControlRigSpline::SetControlPoints(const TArrayView<const FVector>& InPoin
 		}
 
 		// Cache accumulated length at sample array
+		if (SplineData->SamplesArray.Num() > 0)
 		{
 			SplineData->AccumulatedLenth.SetNumUninitialized(SplineData->SamplesArray.Num(), false);
 			SplineData->AccumulatedLenth[0] = 0.f;
diff --git a/Engine/Plugins/Animation/ControlRigSpline/Source/ControlRigSpline/Public/ControlRigSplineUnits.h b/Engine/Plugins/Animation/ControlRigSpline/Source/ControlRigSpline/Public/ControlRigSplineUnits.h
--- a/Engine/Plugins/Animation/ControlRigSpline/Source/ControlRigSpline/Public/ControlRigSplineUnits.h
+++ b/Engine/Plugins/Animation/ControlRigSpline/Source/ControlRigSpline/Public/ControlRigSplineUnits.h
@@ -44,7 +44,7 @@ struct CONTROLRIGSPLINE_API FRigUnit_ControlRigSplineFromPoints : public FRigUni
 	UPROPERTY(meta = (Input))
 	bool bClosed;
 
-	UPROPERTY(meta = (Input))
+	UPROPERTY(meta = (Input, ClampMin = "1"))
 	int32 SamplesPerSegment;
 
 	UPROPERTY(meta = (Input))
