Control Rig: Spline from transforms

diff --git a/Engine/Plugins/Animation/ControlRigSpline/Source/ControlRigSpline/Private/ControlRigSplineTypes.cpp b/Engine/Plugins/Animation/ControlRigSpline/Source/ControlRigSpline/Private/ControlRigSplineTypes.cpp
--- a/Engine/Plugins/Animation/ControlRigSpline/Source/ControlRigSpline/Private/ControlRigSplineTypes.cpp
+++ b/Engine/Plugins/Animation/ControlRigSpline/Source/ControlRigSpline/Private/ControlRigSplineTypes.cpp
@@ -45,6 +45,31 @@ TArray<FVector> ControlRigBaseSpline::GetControlPointsWithoutDuplicates()
 	return ControlPoints;
 }
 
+TArray<uint16> ControlRigBaseSpline::GetControlIndicesWithoutDuplicates()
+{
+	if (bClosed)
+	{
+		const int32 NumPoints = ControlPoints.Num() - Degree;
+		const int32 NumInitPointsToIgnore = FMath::RoundToPositiveInfinity((float)Degree/2) - 1;
+		
+		TArray<uint16> ReducedControlIndices;
+		ReducedControlIndices.Reserve(NumPoints);
+		for (int32 i=0; i<NumPoints; ++i)
+		{
+			ReducedControlIndices.Add(i+NumInitPointsToIgnore);
+		}
+		return ReducedControlIndices;
+	}
+
+	TArray<uint16> ControlIndices;
+	ControlIndices.SetNumUninitialized(ControlPoints.Num());
+	for(int32 i = 0; i<ControlPoints.Num(); ++i)
+	{
+		ControlIndices[i] = i;
+	}
+	return ControlIndices;
+}
+
 void ControlRigBaseSpline::SetControlPoints(const TArrayView<const FVector>& InControlPoints)
 {
 	if (bClosed)
@@ -66,9 +91,9 @@ void ControlRigBaseSpline::SetControlPoints(const TArrayView<const FVector>& InC
 ControlRigBSpline::ControlRigBSpline(const TArrayView<const FVector>& InControlPoints, const uint16 InDegree, const bool bInClosed, bool bInClamped)
 : ControlRigBaseSpline(InControlPoints, InDegree, bInClosed)
 {
-	uint16 N = ControlPoints.Num() - 1;
-	uint16 P = Degree;
-	uint16 M = N + P + 1;
+	const uint16 N = ControlPoints.Num() - 1;
+	const uint16 P = Degree;
+	const uint16 M = N + P + 1;
 	KnotVector.SetNumZeroed(M+1);
 
 	// Clamped means there is p+1 0s at the beginning and 1s at the end of the knot vector
@@ -100,16 +125,16 @@ FVector ControlRigBSpline::GetPointAtParam(float Param)
 	// DeBoor's Algorithm
 	// https://pages.mtu.edu/~shene/COURSES/cs3621/NOTES/spline/B-spline/de-Boor.html
 	
-	float ClampedU = FMath::Clamp(Param, 0.f, 1.f);
+	const float ClampedU = FMath::Clamp(Param, 0.f, 1.f);
 
-	uint16 N = ControlPoints.Num() - 1;
-	uint16 P = Degree;
-	uint16 M = N + P + 1;
+	const uint16 N = ControlPoints.Num() - 1;
+	const uint16 P = Degree;
+	const uint16 M = N + P + 1;
 
 	// Remap U value to KnotVector[P] : KnotVector[M+1 - (P+1)]
-	float Min = KnotVector[P] + KINDA_SMALL_NUMBER;
-	float Max = KnotVector[M-P] - KINDA_SMALL_NUMBER;
-	float U = Min + ClampedU*(Max-Min);
+	const float Min = KnotVector[P] + KINDA_SMALL_NUMBER;
+	const float Max = KnotVector[M-P] - KINDA_SMALL_NUMBER;
+	const float U = Min + ClampedU*(Max-Min);
 
 	// Find [uk, uk+1) where U lives (or k == the first appearence of U)
 	uint16 K = 0;
@@ -163,6 +188,7 @@ ControlRigHermite::ControlRigHermite(const TArrayView<const FVector>& InControlP
 	}
 	else
 	{
+		SegmentPoints.Reserve(ControlPoints.Num()+2);
 		SegmentPoints.Add(ControlPoints[0] + ControlPoints[0] - ControlPoints[1]);
 		SegmentPoints.Append(ControlPoints);
 		SegmentPoints.Add(ControlPoints.Last() + ControlPoints.Last() - ControlPoints[ControlPoints.Num()-2]);
@@ -181,6 +207,7 @@ void ControlRigHermite::SetControlPoints(const TArrayView<const FVector>& InCont
 	}
 	else
 	{
+		SegmentPoints.Reserve(ControlPoints.Num()+2);
 		SegmentPoints.Add(ControlPoints[0] + ControlPoints[0] - ControlPoints[1]);
 		SegmentPoints.Append(ControlPoints);
 		SegmentPoints.Add(ControlPoints.Last() + ControlPoints.Last() - ControlPoints[ControlPoints.Num()-2]);
@@ -195,19 +222,19 @@ static FVector Hermite(const FVector& P0, const FVector& P1, const FVector& M0,
 	const float T2 = T*T;
 	const float T3 = T2*T;
 	
-	float H00 = 2*T3 - 3*T2 + 1;
-	float H10 = T3 - 2*T2 + T;
-	float H01 = -2*T3 + 3*T2;
-	float H11 = T3 - T2;
+	const float H00 = 2*T3 - 3*T2 + 1;
+	const float H10 = T3 - 2*T2 + T;
+	const float H01 = -2*T3 + 3*T2;
+	const float H11 = T3 - T2;
 
 	return H00*P0 + H10*M0 + H01*P1 + H11*M1;
 }
 
 FVector ControlRigHermite::GetPointAtParam(float Param)
 {
-	float ClampedParam = FMath::Clamp(Param, 0.f, 1.f);	
+	const float ClampedParam = FMath::Clamp(Param, 0.f, 1.f);	
 
-	float ParamIntervals = 1.f / (float) NumSegments;
+	const float ParamIntervals = 1.f / (float) NumSegments;
 	uint16 SegmentIndex = FMath::Floor(ClampedParam / ParamIntervals);
 	if (SegmentIndex >= NumSegments)
 	{
@@ -216,12 +243,12 @@ FVector ControlRigHermite::GetPointAtParam(float Param)
 
 	uint16 P3Index = SegmentIndex + 3;
 	P3Index = (P3Index < SegmentPoints.Num()) ? P3Index : SegmentPoints.Num()-1;
-	uint16 P0Index = P3Index - 3;
-	uint16 P1Index = P3Index - 2;
-	uint16 P2Index = P3Index - 1;
+	const uint16 P0Index = P3Index - 3;
+	const uint16 P1Index = P3Index - 2;
+	const uint16 P2Index = P3Index - 1;
 
-	float StartSegmentParam = SegmentIndex * ParamIntervals;
-	float T = FMath::Clamp((ClampedParam-StartSegmentParam)/ParamIntervals, 0.f, 1.f);
+	const float StartSegmentParam = SegmentIndex * ParamIntervals;
+	const float T = FMath::Clamp((ClampedParam-StartSegmentParam)/ParamIntervals, 0.f, 1.f);
 
 	const FVector& P0 = SegmentPoints[P0Index];
 	const FVector& P1 = SegmentPoints[P1Index];
@@ -229,9 +256,9 @@ FVector ControlRigHermite::GetPointAtParam(float Param)
 	const FVector& P3 = SegmentPoints[P3Index];
 
 	// https://www.cs.cmu.edu/~fp/courses/graphics/asst5/catmullRom.pdf
-	float Tension = 0.5f;
-	FVector M1 = Tension * (P2 - P0);
-	FVector M2 = Tension * (P3 - P1);
+	static const float Tension = 0.5f;
+	const FVector M1 = Tension * (P2 - P0);
+	const FVector M2 = Tension * (P3 - P1);
 	return Hermite(P1, P2, M1, M2, T);
 }
 
@@ -247,7 +274,12 @@ FControlRigSplineImpl::~FControlRigSplineImpl()
 TArray<FVector>& FControlRigSplineImpl::GetControlPoints()
 {
 	check(Spline);
-	return Spline->GetControlPoints();	
+	return Spline->GetControlPoints();
+}
+
+TArray<FTransform>& FControlRigSplineImpl::GetControlTransforms()
+{
+	return ControlTransforms;
 }
 
 TArray<FVector> FControlRigSplineImpl::GetControlPointsWithoutDuplicates()
@@ -256,12 +288,30 @@ TArray<FVector> FControlRigSplineImpl::GetControlPointsWithoutDuplicates()
 	return Spline->GetControlPointsWithoutDuplicates();
 }
 
+TArray<FTransform> FControlRigSplineImpl::GetControlTransformsWithoutDuplicates()
+{
+	check(Spline);
+	TArray<uint16> ControlIndices = Spline->GetControlIndicesWithoutDuplicates();
+	TArray<FTransform> Transforms;
+	Transforms.Reserve(ControlIndices.Num());
+	for (const uint16& Index : ControlIndices)
+	{
+		Transforms.Add(ControlTransforms[Index]);
+	}
+	return Transforms;
+}
+
 uint8 FControlRigSplineImpl::GetDegree() const
 {
 	check(Spline);
 	return Spline->GetDegree();
 }
 
+uint16 FControlRigSplineImpl::NumSamples() const
+{
+	return SamplesArray.Num();
+}
+
 FControlRigSpline::FControlRigSpline(const FControlRigSpline& InOther)
 {
 	SplineData = InOther.SplineData;
@@ -284,7 +334,20 @@ uint8 FControlRigSpline::GetDegree() const
 
 void FControlRigSpline::SetControlPoints(const TArrayView<const FVector>& InPoints, const ESplineType SplineMode, const bool bInClosed, const int32 SamplesPerSegment, const float Compression, const float Stretch)
 {
-	const int32 ControlPointsCount = InPoints.Num();
+	TArray<FTransform> Transforms;
+	Transforms.Reserve(InPoints.Num());
+	for (const FVector& Point : InPoints)
+	{
+		FTransform Transform = FTransform::Identity;
+		Transform.SetTranslation(Point);
+		Transforms.Add(Transform);
+	}
+	return SetControlTransforms(Transforms, SplineMode, bInClosed, SamplesPerSegment, Compression, Stretch);
+}
+
+void FControlRigSpline::SetControlTransforms(const TArrayView<const FTransform>& InTransforms, const ESplineType SplineMode, const bool bInClosed, const int32 SamplesPerSegment, const float Compression, const float Stretch)
+{
+	const int32 ControlPointsCount = InTransforms.Num();
 	if (ControlPointsCount < 4)
 	{
 		return;
@@ -300,23 +363,40 @@ void FControlRigSpline::SetControlPoints(const TArrayView<const FVector>& InPoin
 		SplineData = MakeShared<FControlRigSplineImpl>();
 	}
 
-	TArray<FVector> OldControlPoints;
+	TArray<FTransform> OldControlTransforms;
 	if (SplineData->Spline)
 	{
-		OldControlPoints = SplineData->GetControlPointsWithoutDuplicates();
+		OldControlTransforms = SplineData->GetControlTransformsWithoutDuplicates();
 	}
 
-	bool bControlPointsChanged = (SplineData->Spline) ? InPoints != OldControlPoints : true;
-	bool bNumControlPointsChanged = (SplineData->Spline) ? OldControlPoints.Num() != ControlPointsCount : true;
-	bool bSplineModeChanged = SplineMode != SplineData->SplineMode;
-	bool bSamplesCountChanged = SamplesPerSegment != SplineData->SamplesPerSegment;
-	bool bStretchChanged = Stretch != SplineData->Stretch || Compression != SplineData->Compression;
-	bool bClosedChanged = bInClosed != SplineData->bClosed;
+	bool bControlPointsChanged = false;
+	const bool bNumControlPointsChanged = (SplineData->Spline) ? OldControlTransforms.Num() != ControlPointsCount : true;
+	if (!bNumControlPointsChanged)
+	{
+		for (int32 i=0; i<ControlPointsCount; ++i)
+		{
+			if (!InTransforms[i].Equals(OldControlTransforms[i]))
+			{
+				bControlPointsChanged = true;
+				break;
+			}
+		}
+	}
+	else
+	{
+		bControlPointsChanged = true;
+	}
+	
+	const bool bSplineModeChanged = SplineMode != SplineData->SplineMode;
+	const bool bSamplesCountChanged = SamplesPerSegment != SplineData->SamplesPerSegment;
+	const bool bStretchChanged = Stretch != SplineData->Stretch || Compression != SplineData->Compression;
+	const bool bClosedChanged = bInClosed != SplineData->bClosed;
 	if (!bSplineModeChanged && !bControlPointsChanged && !bSamplesCountChanged && !bStretchChanged && !bClosedChanged)
 	{
 		return;
 	}
 
+	SplineData->ControlTransforms = InTransforms;
 	SplineData->SplineMode = SplineMode;
 	SplineData->bClosed = bInClosed;
 	SplineData->SamplesPerSegment = SamplesPerSegment;
@@ -325,7 +405,14 @@ void FControlRigSpline::SetControlPoints(const TArrayView<const FVector>& InPoin
 
 	// If we need to update the spline because the controls points have changed, or the spline mode has changed
 	if (bControlPointsChanged || bSplineModeChanged || bClosedChanged)
-	{	
+	{
+		TArray<FVector> ControlPoints;
+		ControlPoints.Reserve(InTransforms.Num());
+		for (const FTransform& Transform : InTransforms)
+		{
+			ControlPoints.Add(Transform.GetTranslation());
+		}
+		
 		if (bNumControlPointsChanged || bSplineModeChanged || bClosedChanged)
 		{
 			// Delete previously created spline
@@ -339,13 +426,13 @@ void FControlRigSpline::SetControlPoints(const TArrayView<const FVector>& InPoin
 			{
 				case ESplineType::BSpline:
 				{
-					SplineData->Spline = new ControlRigBSpline(InPoints, 3, bInClosed, true);
+					SplineData->Spline = new ControlRigBSpline(ControlPoints, 3, bInClosed, true);
 					break;
 				}
 				case ESplineType::Hermite:
 				{
 					
-					SplineData->Spline = new ControlRigHermite(InPoints, bInClosed);
+					SplineData->Spline = new ControlRigHermite(ControlPoints, bInClosed);
 					break;
 				}
 				default:
@@ -357,7 +444,7 @@ void FControlRigSpline::SetControlPoints(const TArrayView<const FVector>& InPoin
 		}
 		else if (bControlPointsChanged)
 		{
-			SplineData->Spline->SetControlPoints(InPoints);
+			SplineData->Spline->SetControlPoints(ControlPoints);
 		}
 	}
 
@@ -365,14 +452,46 @@ void FControlRigSpline::SetControlPoints(const TArrayView<const FVector>& InPoin
 	if (bControlPointsChanged || bSplineModeChanged || bSamplesCountChanged || bStretchChanged || bClosedChanged)
 	{
 		// Cache sample positions of the spline
-		int32 NumSamples = (ControlPointsCount-1) * SamplesPerSegment;
+		const int32 NumSamples = (ControlPointsCount-1) * SamplesPerSegment;
 		SplineData->SamplesArray.SetNumUninitialized(NumSamples, false);
 			
 		float U = 0.f;
-		float DeltaU = 1.f/(NumSamples-1);
+		const float DeltaU = 1.f/(NumSamples-1);
+		for (int32 i=0; i<NumSamples; ++i, U+=DeltaU)
+		{
+			SplineData->SamplesArray[i].SetTranslation(SplineData->Spline->GetPointAtParam(U));
+		}
+
+		FVector UpVector = ControlPointsCount > 0 ? InTransforms[0].GetRotation().GetUpVector() : FVector(0,0,0);
+		UpVector.Normalize();
+
+		U = 0.f;
 		for (int32 i=0; i<NumSamples; ++i, U+=DeltaU)
 		{
-			SplineData->SamplesArray[i] = SplineData->Spline->GetPointAtParam(U); 
+			FVector Tangent = TangentAtParam(U);
+			Tangent.Normalize();
+			
+			int32 PrevIndex = U*(ControlPointsCount-1);
+			int32 NextIndex = PrevIndex+1;
+			if (!InTransforms.IsValidIndex(NextIndex))
+			{
+				PrevIndex--;
+				NextIndex--;
+			}
+			
+			// Interpolate inside segment
+			const float UPrev = PrevIndex / (float)(ControlPointsCount-1);
+			const float UNext = NextIndex / (float)(ControlPointsCount-1);
+			const float Interp = (U - UPrev) / (UNext - UPrev);
+			const FQuat InterpRotation = FQuat::Slerp(InTransforms[PrevIndex].GetRotation(), InTransforms[NextIndex].GetRotation(), Interp);
+			
+			//FVector Binormal = FVector::CrossProduct(Tangent, UpVector);
+			FMatrix RotationMatrix = FRotationMatrix::MakeFromXZ(Tangent, InterpRotation.GetUpVector());
+			FQuat RotationQuat = RotationMatrix.ToQuat();
+			RotationQuat.Normalize();
+			SplineData->SamplesArray[i].SetRotation(RotationQuat);
+
+			SplineData->SamplesArray[i].SetScale3D(FMath::Lerp(InTransforms[PrevIndex].GetScale3D(), InTransforms[NextIndex].GetScale3D(), Interp));
 		}
 
 		// Correct length of samples
@@ -380,7 +499,7 @@ void FControlRigSpline::SetControlPoints(const TArrayView<const FVector>& InPoin
 		{
 			if (SplineData->InitialLengths.Num() > 0)
 			{
-				TArray<FVector> SamplesBeforeCorrect = SplineData->SamplesArray;
+				const TArray<FTransform> SamplesBeforeCorrect = SplineData->SamplesArray;
 				for (int32 i = 0; i < ControlPointsCount - 1; ++i)
 				{
 					for (int32 j = (i == 0) ? 1 : 0; j < SamplesPerSegment; ++j)
@@ -393,17 +512,17 @@ void FControlRigSpline::SetControlPoints(const TArrayView<const FVector>& InPoin
 						}
 						
 						// Get direction from samples before correction
-						FVector Dir = SamplesBeforeCorrect[CurSampleIndex] - SamplesBeforeCorrect[PrevSampleIndex];
+						FVector Dir = SamplesBeforeCorrect[CurSampleIndex].GetTranslation() - SamplesBeforeCorrect[PrevSampleIndex].GetTranslation();
 						Dir.Normalize();
 
 						float InitialLength = SplineData->InitialLengths[PrevSampleIndex];
 						// Current length as the projection on the Dir vector (might be negative)
-						float CurrentLength = (SplineData->SamplesArray[CurSampleIndex] - SplineData->SamplesArray[PrevSampleIndex]).Dot(Dir);
+						float CurrentLength = (SplineData->SamplesArray[CurSampleIndex].GetTranslation() - SplineData->SamplesArray[PrevSampleIndex].GetTranslation()).Dot(Dir);
 						float FixedLength = FMath::Clamp(CurrentLength,
 							(Compression > 0.f) ? InitialLength * Compression : CurrentLength,
 							(Stretch > 0.f) ? InitialLength * Stretch : CurrentLength);
 
-						SplineData->SamplesArray[CurSampleIndex] = SplineData->SamplesArray[PrevSampleIndex] + Dir * FixedLength;
+						SplineData->SamplesArray[CurSampleIndex].SetTranslation(SplineData->SamplesArray[PrevSampleIndex].GetTranslation() + Dir * FixedLength);
 					}
 				}
 			}
@@ -420,7 +539,7 @@ void FControlRigSpline::SetControlPoints(const TArrayView<const FVector>& InPoin
 			}
 			for (int32 i = 1; i < SplineData->SamplesArray.Num(); ++i)
 			{
-				float CurrentLength = FVector::Distance(SplineData->SamplesArray[i - 1], SplineData->SamplesArray[i]);
+				float CurrentLength = FVector::Distance(SplineData->SamplesArray[i - 1].GetTranslation(), SplineData->SamplesArray[i].GetTranslation());
 				if (bSplineModeChanged || bClosedChanged || bSamplesCountChanged || bNumControlPointsChanged)
 				{
 					SplineData->InitialLengths[i-1] = CurrentLength;
@@ -428,7 +547,7 @@ void FControlRigSpline::SetControlPoints(const TArrayView<const FVector>& InPoin
 				SplineData->AccumulatedLenth[i] = SplineData->AccumulatedLenth[i - 1] + CurrentLength;
 			}
 		}
-	}	
+	}
 }
 
 FVector FControlRigSpline::PositionAtParam(const float InParam) const
@@ -438,15 +557,38 @@ FVector FControlRigSpline::PositionAtParam(const float InParam) const
 		return FVector();
 	}
 
-	float ClampedU = FMath::Clamp<float>(InParam, 0.f, 1.f);
+	const float ClampedU = FMath::Clamp<float>(InParam, 0.f, 1.f);
 
 	const int32 LastIndex = SplineData->SamplesArray.Num() - 1;
-	float fIndexPrev = ClampedU * LastIndex;	
-	int32 IndexPrev = FMath::Floor<int32>(fIndexPrev);
-	float ULocal = fIndexPrev - IndexPrev;
-	int32 IndexNext = (IndexPrev < LastIndex) ? IndexPrev + 1 : IndexPrev;
+	const float fIndexPrev = ClampedU * LastIndex;	
+	const int32 IndexPrev = FMath::Floor<int32>(fIndexPrev);
+	const float ULocal = fIndexPrev - IndexPrev;
+	const int32 IndexNext = (IndexPrev < LastIndex) ? IndexPrev + 1 : IndexPrev;
+
+	return SplineData->SamplesArray[IndexPrev].GetTranslation() * (1.f - ULocal) + SplineData->SamplesArray[IndexNext].GetTranslation() * ULocal;
+}
+
+FTransform FControlRigSpline::TransformAtParam(const float InParam) const
+{
+	if (!SplineData.IsValid())
+	{
+		return FTransform();
+	}
+
+	const float ClampedU = FMath::Clamp<float>(InParam, 0.f, 1.f);
 
-	return SplineData->SamplesArray[IndexPrev] * (1.f - ULocal) + SplineData->SamplesArray[IndexNext] * ULocal;
+	const int32 LastIndex = SplineData->SamplesArray.Num() - 1;
+	const float fIndexPrev = ClampedU * LastIndex;	
+	const int32 IndexPrev = FMath::Floor<int32>(fIndexPrev);
+	const float ULocal = fIndexPrev - IndexPrev;
+	const int32 IndexNext = (IndexPrev < LastIndex) ? IndexPrev + 1 : IndexPrev;
+
+	FTransform Result;
+	FQuat Rotation = FQuat::Slerp(SplineData->SamplesArray[IndexPrev].GetRotation(), SplineData->SamplesArray[IndexNext].GetRotation(), ULocal);
+	Rotation.Normalize();
+	Result.SetRotation(Rotation);
+	Result.LerpTranslationScale3D(SplineData->SamplesArray[IndexPrev], SplineData->SamplesArray[IndexNext], ScalarRegister(ULocal));
+	return Result;
 }
 
 FVector FControlRigSpline::TangentAtParam(const float InParam) const
@@ -458,7 +600,25 @@ FVector FControlRigSpline::TangentAtParam(const float InParam) const
 
 	const float ClampedU = FMath::Clamp<float>(InParam, 0.f, 1.f);
 
-	int32 IndexPrev = ClampedU * (SplineData->SamplesArray.Num()-2);
-	return SplineData->SamplesArray[IndexPrev+1] - SplineData->SamplesArray[IndexPrev];
+	const int32 IndexPrev = ClampedU * (SplineData->SamplesArray.Num()-2);
+	return SplineData->SamplesArray[IndexPrev+1].GetTranslation() - SplineData->SamplesArray[IndexPrev].GetTranslation();
+}
+
+float FControlRigSpline::LengthAtParam(const float InParam) const
+{
+	if (!SplineData.IsValid())
+	{
+		return 0.f;
+	}
+
+	const float ClampedU = FMath::Clamp<float>(InParam, 0.f, 1.f);
+
+	const int32 LastIndex = SplineData->SamplesArray.Num() - 1;
+	const float fIndexPrev = ClampedU * LastIndex;	
+	const int32 IndexPrev = FMath::Floor<int32>(fIndexPrev);
+	const float ULocal = fIndexPrev - IndexPrev;
+	const int32 IndexNext = (IndexPrev < LastIndex) ? IndexPrev + 1 : IndexPrev;
+
+	return SplineData->AccumulatedLenth[IndexPrev] * (1.f - ULocal) + SplineData->AccumulatedLenth[IndexNext] * ULocal;
 }
 
diff --git a/Engine/Plugins/Animation/ControlRigSpline/Source/ControlRigSpline/Private/ControlRigSplineUnits.cpp b/Engine/Plugins/Animation/ControlRigSpline/Source/ControlRigSpline/Private/ControlRigSplineUnits.cpp
--- a/Engine/Plugins/Animation/ControlRigSpline/Source/ControlRigSpline/Private/ControlRigSplineUnits.cpp
+++ b/Engine/Plugins/Animation/ControlRigSpline/Source/ControlRigSpline/Private/ControlRigSplineUnits.cpp
@@ -23,6 +23,21 @@ FRigUnit_ControlRigSplineFromPoints_Execute()
 	Spline.SetControlPoints(PointsView, SplineMode, bClosed, SamplesPerSegment, Compression, Stretch);
 }
 
+FRigUnit_ControlRigSplineFromTransforms_Execute()
+{
+	// reset the spline
+	Spline = FControlRigSpline();
+
+	if (Transforms.Num() < 4)
+	{
+		UE_CONTROLRIG_RIGUNIT_REPORT_WARNING(TEXT("Cannot create a spline with less than 4 transforms (%d received)."), Transforms.Num());
+		return;
+	}
+		
+	const TArrayView<const FTransform> TransformsView(Transforms.GetData(), Transforms.Num());
+	Spline.SetControlTransforms(TransformsView, SplineMode, bClosed, SamplesPerSegment, Compression, Stretch);
+}
+
 FRigUnit_SetSplinePoints_Execute()
 {
 	if (!Spline.SplineData.IsValid())
@@ -35,19 +50,19 @@ FRigUnit_SetSplinePoints_Execute()
 	{
 		UE_LOG(LogControlRig, Error, TEXT("Invalid input spline implementation."));
 		return;
-	}	
+	}
 
 	if (!Spline.SplineData->bClosed)
 	{
-		if (Points.Num() != Spline.SplineData->GetControlPoints().Num())
+		if (Points.Num() != Spline.SplineData->GetControlTransforms().Num())
 		{
 			UE_LOG(LogControlRig, Error, TEXT("Number of input points does not match the number of point in the spline."));
 			return;
 		}
 	}
-	else if(Points.Num() != Spline.SplineData->GetControlPoints().Num() - Spline.GetDegree())
+	else if(Points.Num() != Spline.SplineData->GetControlTransforms().Num() - Spline.GetDegree())
 	{
-		if (Points.Num() != Spline.SplineData->GetControlPoints().Num())
+		if (Points.Num() != Spline.SplineData->GetControlTransforms().Num())
 		{
 			UE_LOG(LogControlRig, Error, TEXT("Number of input points does not match the number of point in the spline."));
 			return;
@@ -58,6 +73,41 @@ FRigUnit_SetSplinePoints_Execute()
 	Spline.SetControlPoints(PointsView, Spline.SplineData->SplineMode, Spline.SplineData->bClosed, Spline.SplineData->SamplesPerSegment, Spline.SplineData->Compression, Spline.SplineData->Stretch);
 }
 
+FRigUnit_SetSplineTransforms_Execute()
+{
+	if (!Spline.SplineData.IsValid())
+	{
+		UE_LOG(LogControlRig, Error, TEXT("Invalid input spline."));
+		return;
+	}
+	
+	if (Spline.SplineData->Spline == nullptr)
+	{
+		UE_LOG(LogControlRig, Error, TEXT("Invalid input spline implementation."));
+		return;
+	}
+
+	if (!Spline.SplineData->bClosed)
+	{
+		if (Transforms.Num() != Spline.SplineData->GetControlTransforms().Num())
+		{
+			UE_LOG(LogControlRig, Error, TEXT("Number of input points does not match the number of point in the spline."));
+			return;
+		}
+	}
+	else if(Transforms.Num() != Spline.SplineData->GetControlTransforms().Num() - Spline.GetDegree())
+	{
+		if (Transforms.Num() != Spline.SplineData->GetControlTransforms().Num())
+		{
+			UE_LOG(LogControlRig, Error, TEXT("Number of input points does not match the number of point in the spline."));
+			return;
+		}
+	}
+
+	const TArrayView<const FTransform> TransformsView(Transforms.GetData(), Transforms.Num());
+	Spline.SetControlTransforms(TransformsView, Spline.SplineData->SplineMode, Spline.SplineData->bClosed, Spline.SplineData->SamplesPerSegment, Spline.SplineData->Compression, Spline.SplineData->Stretch);
+}
+
 FRigUnit_PositionFromControlRigSpline_Execute()
 {
 	if (!Spline.SplineData.IsValid())
@@ -86,27 +136,10 @@ FRigUnit_TransformFromControlRigSpline_Execute()
 		UE_LOG(LogControlRig, Error, TEXT("Invalid input spline implementation."));
 		return;
 	}
-	
-	FVector UpVectorNormalized = UpVector;
-	UpVectorNormalized.Normalize();
 
 	const float ClampedU = FMath::Clamp<float>(U, 0.f, 1.f);
-	const float ClampedRoll = FMath::Clamp<float>(Roll, -180.f, 180.f);
-
-	FVector Tangent = Spline.TangentAtParam(ClampedU);
+	Transform = Spline.TransformAtParam(ClampedU);
 
-	// Check if Tangent can be normalized. If not, keep the same tangent as before.
-	if (!Tangent.Normalize())
-	{
-		Tangent = Transform.ToMatrixNoScale().GetUnitAxis(EAxis::X);
-	}
-	FVector Binormal = FVector::CrossProduct(Tangent, UpVectorNormalized);
-	Binormal = Binormal.RotateAngleAxis(ClampedRoll * ClampedU, Tangent);
-	
-	FMatrix RotationMatrix = FRotationMatrix::MakeFromXZ(Tangent, Binormal);
-
-	Transform.SetFromMatrix(RotationMatrix);
-	Transform.SetTranslation(Spline.PositionAtParam(U));
 }
 
 FRigUnit_TangentFromControlRigSpline_Execute()
@@ -189,6 +222,22 @@ FRigUnit_GetLengthControlRigSpline_Execute()
 	Length = Spline.SplineData->AccumulatedLenth.Last();
 }
 
+FRigUnit_GetLengthAtParamControlRigSpline_Execute()
+{
+	if (!Spline.SplineData.IsValid())
+	{
+		return;
+	}
+
+	if (Spline.SplineData->Spline == nullptr)
+	{
+		UE_LOG(LogControlRig, Error, TEXT("Invalid input spline implementation."));
+		return;
+	}
+	
+	Length = Spline.LengthAtParam(U);
+}
+
 FRigUnit_FitChainToSplineCurve_Execute()
 {
 	FRigUnit_FitChainToSplineCurveItemArray::StaticExecute(
@@ -232,6 +281,11 @@ FRigVMStructUpgradeInfo FRigUnit_FitChainToSplineCurve::GetUpgradeInfo() const
 	return FRigVMStructUpgradeInfo(*this, NewNode);
 }
 
+FRigUnit_FitSplineCurveToChain_Execute()
+{
+	FRigUnit_FitSplineCurveToChainItemArray::StaticExecute(ExecuteContext, Items.Keys, Spline);
+}
+
 FRigUnit_FitChainToSplineCurveItemArray_Execute()
 {
     DECLARE_SCOPE_HIERARCHICAL_COUNTER_RIGUNIT()
@@ -466,7 +520,7 @@ FRigUnit_FitChainToSplineCurveItemArray_Execute()
 
 		float DistanceA = (LastPosition - A).Size();
 		float DistanceB = (LastPosition - B).Size();
-		
+
 		if (DistanceB > BoneLength)
 		{
 			float Ratio = BoneLength / DistanceB;
@@ -612,6 +666,48 @@ FRigUnit_FitChainToSplineCurveItemArray_Execute()
 			Hierarchy->SetGlobalTransform(CachedItems[Index], BaseTransform, bPropagateToChildren);
 		}
 	}
+	else
+	{
+		FTransform BaseTransform = FTransform::Identity;
+		FRigElementKey ParentKey = Hierarchy->GetFirstParent(CachedItems[0].GetKey());
+		if (ParentKey.IsValid())
+		{
+			BaseTransform = Hierarchy->GetGlobalTransform(ParentKey);
+		}
+
+		for (int32 Index = 0; Index < CachedItems.Num(); Index++)
+		{
+			ItemLocalTransforms[Index] = Hierarchy->GetLocalTransform(CachedItems[Index]);
+		}
+
+		for (int32 Index = 0; Index < CachedItems.Num(); Index++)
+		{
+			if (ItemRotationA[Index] >= Rotations.Num() ||
+				ItemRotationB[Index] >= Rotations.Num())
+			{
+				continue;
+			}
+
+			FQuat Rotation = Rotations[ItemRotationA[Index]].Rotation;
+			FQuat RotationB = Rotations[ItemRotationB[Index]].Rotation;
+			if (ItemRotationA[Index] != ItemRotationB[Index])
+			{
+				if (ItemRotationT[Index] > 1.f - SMALL_NUMBER)
+				{
+					Rotation = RotationB;
+				}
+				else if (ItemRotationT[Index] > SMALL_NUMBER)
+				{
+					Rotation = FQuat::Slerp(Rotation, RotationB, ItemRotationT[Index]).GetNormalized();
+				}
+			}
+
+			BaseTransform = ItemLocalTransforms[Index] * BaseTransform;
+			Rotation = FQuat::Slerp(BaseTransform.GetRotation(), BaseTransform.GetRotation() * Rotation, FMath::Clamp<float>(Weight, 0.f, 1.f));
+			BaseTransform.SetRotation(Rotation);
+			Hierarchy->SetGlobalTransform(CachedItems[Index], BaseTransform, bPropagateToChildren);
+		}
+	}
 
 	if (ExecuteContext.GetDrawInterface() != nullptr && DebugSettings.bEnabled)
 	{
@@ -639,9 +735,141 @@ FRigUnit_FitChainToSplineCurveItemArray_Execute()
 	}
 }
 
-FRigUnit_FitSplineCurveToChain_Execute()
+FRigUnit_SplineConstraint_Execute()
 {
-	FRigUnit_FitSplineCurveToChainItemArray::StaticExecute(ExecuteContext, Items.Keys, Spline);
+	DECLARE_SCOPE_HIERARCHICAL_COUNTER_RIGUNIT()
+
+	URigHierarchy* Hierarchy = ExecuteContext.Hierarchy;
+	if (Hierarchy == nullptr)
+	{
+		return;
+	}
+
+	if (!Spline.SplineData.IsValid())
+	{
+		return;
+	}
+
+	if (Spline.SplineData->Spline == nullptr)
+	{
+		UE_LOG(LogControlRig, Error, TEXT("Invalid input spline implementation."));
+		return;
+	}
+
+	const float ClampedMin = FMath::Min(FMath::Max(Minimum, 0.f), 1.f);
+	float ClampedMax = FMath::Min(FMath::Max(Maximum, 0.f), 1.f);
+	ClampedMax = FMath::Max(ClampedMax, ClampedMin);
+
+	float& ChainLength = WorkData.ChainLength;
+	TArray<FTransform>& ItemTransforms = WorkData.ItemTransforms;
+	TArray<float>& ItemSegments = WorkData.ItemSegments;
+	TArray<FCachedRigElement>& CachedItems = WorkData.CachedItems;
+
+	if(CachedItems.Num() == 0 && Items.Num() > 1)
+	{
+		ItemTransforms.Reset();
+		ItemSegments.Reset();
+
+		for (FRigElementKey Item : Items)
+		{
+			CachedItems.Add(FCachedRigElement(Item, Hierarchy));
+		}
+
+		if (CachedItems.Num() < 2)
+		{
+			UE_CONTROLRIG_RIGUNIT_REPORT_WARNING(TEXT("Didn't find enough bones. You need at least two in the chain!"));
+			return;
+		}
+
+		ItemTransforms.SetNumZeroed(CachedItems.Num());
+		ItemSegments.SetNumZeroed(CachedItems.Num());
+		ChainLength = 0.f;
+		ItemSegments[0] = 0.f;
+		for (int32 Index = 1; Index < CachedItems.Num(); Index++)
+		{
+			FVector A = Hierarchy->GetGlobalTransform(CachedItems[Index - 1]).GetLocation();
+			FVector B = Hierarchy->GetGlobalTransform(CachedItems[Index]).GetLocation();
+			ItemSegments[Index] = (A - B).Size();
+			ChainLength += ItemSegments[Index];
+		}
+	}
+
+	if (CachedItems.Num() < 2)
+	{
+		return;
+	}
+
+	TArray<float> InitialLengths;
+	if (Alignment == EControlRigCurveAlignment::Front)
+	{
+		ItemTransforms[0] = Spline.TransformAtParam(ClampedMin);
+		float PrevParam = ClampedMin;
+		float PrevLength = Spline.LengthAtParam(PrevParam);
+		const float TotalSplineLength = Spline.LengthAtParam(1.f);
+		const int32 LastIndex = Spline.SplineData->NumSamples() - 1;
+		float AccumulatedLength = Spline.LengthAtParam(ClampedMin);
+		for (int32 SegmentIndex = 1; SegmentIndex < ItemSegments.Num(); ++SegmentIndex)
+		{
+			const float& SegmentLenth = ItemSegments[SegmentIndex];
+			AccumulatedLength += SegmentLenth;
+			FTransform& Transform = ItemTransforms[SegmentIndex];
+			bool bFoundSample = false;
+			if (AccumulatedLength < TotalSplineLength)
+			{
+				const float fStartSampleSearch = PrevParam * LastIndex;	
+				const int32 StartSampleSearch = FMath::Floor<int32>(fStartSampleSearch);
+
+				// Find the sample interval where the bone should be
+				for(int32 SampleIndex = StartSampleSearch; SampleIndex <= LastIndex; SampleIndex++)
+				{
+					const float CurrentSegmentLength = Spline.SplineData->AccumulatedLenth[SampleIndex];
+					if (CurrentSegmentLength > AccumulatedLength)
+					{
+						const int32 IndexPrev = SampleIndex-1;
+						const int32 IndexNext = SampleIndex;
+						const float& IndexPrevLength = Spline.SplineData->AccumulatedLenth[IndexPrev];
+						const float& IndexNextLength = Spline.SplineData->AccumulatedLenth[IndexNext];
+						const float ULocal = (AccumulatedLength - IndexPrevLength) / (IndexNextLength - IndexPrevLength);
+						check(ULocal >= 0.f && ULocal <= 1.f);
+						const float UPrev = IndexPrev / (float)LastIndex;
+						const float UNext = IndexNext / (float)LastIndex;
+						PrevParam = UPrev + (ULocal * (UNext - UPrev));
+						check(PrevParam >= ClampedMin && PrevParam <= 1.f);
+						PrevLength = (ULocal-1.f) * IndexPrevLength + ULocal * IndexNextLength;
+
+						Transform = Spline.TransformAtParam(PrevParam);
+						bFoundSample = true;
+						break;
+					}
+				}
+			}
+			if (!bFoundSample)
+			{
+				Transform = ItemTransforms.IsValidIndex(SegmentIndex-1) ? ItemTransforms[SegmentIndex-1] : FTransform::Identity;
+				Transform.SetLocation(Transform.GetLocation() + SegmentLenth*Transform.GetRotation().GetForwardVector());
+			}
+			check(Transform.GetRotation().IsNormalized());
+		}
+	}
+	else if (Alignment == EControlRigCurveAlignment::Stretched)
+	{
+		ItemTransforms[0] = Spline.TransformAtParam(ClampedMin);
+		ItemTransforms.Last() = Spline.TransformAtParam(ClampedMax);
+
+		float AccumulatedLength = 0.f;
+		for (int32 SegmentIndex=0; SegmentIndex < ItemSegments.Num(); ++SegmentIndex)
+		{
+			AccumulatedLength += ItemSegments[SegmentIndex];
+			const float U = AccumulatedLength*ClampedMax/ChainLength;
+			const float LocalParam = (1.f - U)*ClampedMin + U*ClampedMax;
+			ItemTransforms[SegmentIndex] = Spline.TransformAtParam(LocalParam);
+		}
+	}
+
+	for (int32 Index = 0; Index < CachedItems.Num(); Index++)
+	{
+		Hierarchy->SetGlobalTransform(CachedItems[Index], ItemTransforms[Index], bPropagateToChildren);
+	}
 }
 
 FRigVMStructUpgradeInfo FRigUnit_FitSplineCurveToChain::GetUpgradeInfo() const
@@ -683,36 +911,36 @@ FRigUnit_FitSplineCurveToChainItemArray_Execute()
 	const bool bClosed = Spline.SplineData->bClosed;
 
 	// 1.- Create spline from chain
-	TArray<FVector> AuxControlPoints;
-	AuxControlPoints.SetNumUninitialized(Items.Num());	
+	TArray<FTransform> AuxControlTransforms;
+	AuxControlTransforms.SetNumUninitialized(Items.Num());	
 	for (int32 i = 0; i < Items.Num(); ++i)
 	{
 		const FRigElementKey& Key = Items[i];
-		AuxControlPoints[i] = Hierarchy->GetGlobalTransform(Key).GetLocation();
+		AuxControlTransforms[i] = Hierarchy->GetGlobalTransform(Key);
 	}
     FControlRigSpline AuxSpline;
-	const TArrayView<const FVector> PointsView(AuxControlPoints.GetData(), AuxControlPoints.Num());
-	AuxSpline.SetControlPoints(PointsView, ESplineType::Hermite, bClosed, Spline.SplineData->SamplesPerSegment);
+	const TArrayView<const FTransform> TransformsView(AuxControlTransforms.GetData(), AuxControlTransforms.Num());
+	AuxSpline.SetControlTransforms(TransformsView, ESplineType::Hermite, bClosed, Spline.SplineData->SamplesPerSegment);
 	
 
 	// 2.-  for each control point in the original spline
 	//			figure out its u in the original spline (preprocess)
 	//			query position at u on the new spline
-	const TArray<FVector> ControlPoints = Spline.SplineData->GetControlPointsWithoutDuplicates();
-	const int32 NumControlPoints = ControlPoints.Num();
-	TArray<FVector> NewControlPoints;
-	NewControlPoints.SetNumUninitialized(NumControlPoints);
+	const TArray<FTransform> ControlTransforms = Spline.SplineData->GetControlTransformsWithoutDuplicates();
+	const int32 NumControlPoints = ControlTransforms.Num();
+	TArray<FTransform> NewControlTransforms;
+	NewControlTransforms.SetNumUninitialized(NumControlPoints);
 	float U = 0;
 	const float DeltaU = bClosed ? 1 / (float)(NumControlPoints) : 1 / (float)(NumControlPoints - 1);
 	for (int32 i = 0; i < NumControlPoints; ++i)
 	{
-		FVector NewPosition = AuxSpline.PositionAtParam(U);
-		NewControlPoints[i] = NewPosition;
+		FTransform NewPosition = AuxSpline.TransformAtParam(U);
+		NewControlTransforms[i] = NewPosition;
 		U += DeltaU;
 	}
 
-	const TArrayView<const FVector> NewPointsView(NewControlPoints.GetData(), NewControlPoints.Num());
-	Spline.SetControlPoints(NewPointsView, Spline.SplineData->SplineMode, bClosed, Spline.SplineData->SamplesPerSegment, Spline.SplineData->Compression, Spline.SplineData->Stretch);
+	const TArrayView<const FTransform> NewTransformsView(NewControlTransforms.GetData(), NewControlTransforms.Num());
+	Spline.SetControlTransforms(NewTransformsView, Spline.SplineData->SplineMode, bClosed, Spline.SplineData->SamplesPerSegment, Spline.SplineData->Compression, Spline.SplineData->Stretch);
 }
 
 FRigUnit_ClosestParameterFromControlRigSpline_Execute()
@@ -733,8 +961,8 @@ FRigUnit_ClosestParameterFromControlRigSpline_Execute()
 	auto DistanceToSegment = [&](int32 Index0, int32 Index1, float& OutParam) -> float
 	{
 		const int32 MaxIndex = Spline.SplineData->SamplesArray.Num()-1;
-		const FVector& P0 = Spline.SplineData->SamplesArray[Index0];
-		const FVector& P1 = Spline.SplineData->SamplesArray[Index1];
+		const FVector& P0 = Spline.SplineData->SamplesArray[Index0].GetTranslation();
+		const FVector& P1 = Spline.SplineData->SamplesArray[Index1].GetTranslation();
 
 		const FVector V = P1 - P0;
 		const FVector U = P0 - Position;
@@ -768,7 +996,7 @@ FRigUnit_ClosestParameterFromControlRigSpline_Execute()
 			ClosestDistance = Distance;
 			U = Param;
 		}
-	}				
+	}
 }
 
 FRigUnit_ParameterAtPercentage_Execute()
@@ -811,7 +1039,7 @@ FRigUnit_ParameterAtPercentage_Execute()
 	float LengthNext = Spline.SplineData->AccumulatedLenth[NextIndex];
 
 	float Interp = (SearchLength - LengthPrev) / (LengthNext - LengthPrev);
-	U = (1 - Interp) * UPrev + Interp * UNext;			
+	U = (1 - Interp) * UPrev + Interp * UNext;
 }
 
 
diff --git a/Engine/Plugins/Animation/ControlRigSpline/Source/ControlRigSpline/Public/ControlRigSplineTypes.h b/Engine/Plugins/Animation/ControlRigSpline/Source/ControlRigSpline/Public/ControlRigSplineTypes.h
--- a/Engine/Plugins/Animation/ControlRigSpline/Source/ControlRigSpline/Public/ControlRigSplineTypes.h
+++ b/Engine/Plugins/Animation/ControlRigSpline/Source/ControlRigSpline/Public/ControlRigSplineTypes.h
@@ -54,6 +54,8 @@ public:
 	
 	virtual TArray<FVector> GetControlPointsWithoutDuplicates();
 
+	virtual TArray<uint16> GetControlIndicesWithoutDuplicates();
+
 	uint8 GetDegree() const { return Degree; }
 };
 
@@ -102,6 +104,9 @@ struct CONTROLRIGSPLINE_API FControlRigSplineImpl
 	// Spline type
 	ESplineType SplineMode;
 
+	// The transforms of the control points
+	TArray<FTransform> ControlTransforms;
+
 	// The initial lengths between samples
 	TArray<float> InitialLengths;
 
@@ -121,7 +126,7 @@ struct CONTROLRIGSPLINE_API FControlRigSplineImpl
 	float Stretch;
 
 	// Positions along the "real" curve (no samples in the first and last segments of a hermite spline)
-	TArray<FVector> SamplesArray;
+	TArray<FTransform> SamplesArray;
 
 	// Accumulated length along the spline given by samples
 	TArray<float> AccumulatedLenth;
@@ -129,11 +134,20 @@ struct CONTROLRIGSPLINE_API FControlRigSplineImpl
 	// Returns a reference to the control points that were used to create this spline 
 	TArray<FVector>& GetControlPoints();
 
+	// Returns a reference to the control transforms that were used to create this spline 
+	TArray<FTransform>& GetControlTransforms();
+
 	// Returns the control points that were used to create this spline, removing the duplicates in case of a closed spline
 	TArray<FVector> GetControlPointsWithoutDuplicates();
+	
+	// Returns the control transforms that were used to create this spline, removing the duplicates in case of a closed spline
+	TArray<FTransform> GetControlTransformsWithoutDuplicates();
 
 	// Returns the degree of the curve
 	uint8 GetDegree() const;
+
+	// Returns the total number of samples cached
+	uint16 NumSamples() const;
 };
 
 USTRUCT(BlueprintType)
@@ -158,7 +172,7 @@ struct CONTROLRIGSPLINE_API FControlRigSpline
 	uint8 GetDegree() const;
 
 	/**
-	* Sets the control points in the spline. It will build the spline if needed, or forceRebuild is true,
+	* Sets the control points in the spline. It will build the spline if needed, 
 	* or will update the points if building from scratch is not necessary. The type of spline to build will
 	* depend on what is set in SplineMode.
 	*
@@ -171,6 +185,20 @@ struct CONTROLRIGSPLINE_API FControlRigSpline
 	*/
 	void SetControlPoints(const TArrayView<const FVector>& InPoints, const ESplineType SplineMode = ESplineType::BSpline, const bool bInClosed = false, const int32 SamplesPerSegment = 16, const float Compression = 1.f, const float Stretch = 1.f);
 
+	/**
+	* Sets the control transforms in the spline. It will build the spline if needed, 
+	* or will update the transforms if building from scratch is not necessary. The type of spline to build will
+	* depend on what is set in SplineMode.
+	*
+	* @param InTransforms	The control transforms to set.
+	* @param SplineMode	The type of spline
+	* @param bInClosed	If the spline should be closed
+	* @param SamplesPerSegment The samples to cache for every segment defined between two control rig
+	* @param Compression The allowed length compression (1.f being do not allow compression`). If 0, no restriction wil be applied.
+	* @param Stretch The allowed length stretch (1.f being do not allow stretch). If 0, no restriction wil be applied.
+	*/
+	void SetControlTransforms(const TArrayView<const FTransform>& InTransforms, const ESplineType SplineMode = ESplineType::BSpline, const bool bInClosed = false, const int32 SamplesPerSegment = 16, const float Compression = 1.f, const float Stretch = 1.f);
+
 	/**
 	* Given an InParam float in [0, 1], will return the position of the spline at that point.
 	*
@@ -178,6 +206,14 @@ struct CONTROLRIGSPLINE_API FControlRigSpline
 	* @return			The position in the spline.
 	*/
 	FVector PositionAtParam(const float InParam) const;
+	
+	/**
+	* Given an InParam float in [0, 1], will return the transform of the spline at that point.
+	*
+	* @param InParam	The parameter between [0, 1] to query.
+	* @return			The transform in the spline.
+	*/
+	FTransform TransformAtParam(const float InParam) const;
 
 	/**
 	* Given an InParam float in [0, 1], will return the tangent vector of the spline at that point. 
@@ -187,4 +223,12 @@ struct CONTROLRIGSPLINE_API FControlRigSpline
 	* @return			The tangent of the spline at InParam.
 	*/
 	FVector TangentAtParam(const float InParam) const;
+
+	/**
+	* Given an InParam float in [0, 1], will return the length of the spline at that point. 
+	*
+	* @param InParam	The parameter between [0, 1] to query.
+	* @return			The length of the spline at InParam.
+	*/
+	float LengthAtParam(const float InParam) const;
 };
diff --git a/Engine/Plugins/Animation/ControlRigSpline/Source/ControlRigSpline/Public/ControlRigSplineUnits.h b/Engine/Plugins/Animation/ControlRigSpline/Source/ControlRigSpline/Public/ControlRigSplineUnits.h
--- a/Engine/Plugins/Animation/ControlRigSpline/Source/ControlRigSpline/Public/ControlRigSplineUnits.h
+++ b/Engine/Plugins/Animation/ControlRigSpline/Source/ControlRigSpline/Public/ControlRigSplineUnits.h
@@ -57,6 +57,49 @@ struct CONTROLRIGSPLINE_API FRigUnit_ControlRigSplineFromPoints : public FRigUni
 	FControlRigSpline Spline;
 };
 
+/*
+ * Creates a Spline curve from an array of positions
+ */
+USTRUCT(meta = (DisplayName = "Spline From Transforms", Keywords="Spline From Transforms", Category = "Splines", Varying, NodeColor="0.737911 0.099899 0.099899"))
+struct CONTROLRIGSPLINE_API FRigUnit_ControlRigSplineFromTransforms : public FRigUnit_ControlRigSplineBase
+{
+	GENERATED_BODY()
+
+	FRigUnit_ControlRigSplineFromTransforms()
+	{
+		SplineMode = ESplineType::Hermite;
+		bClosed = false;
+		SamplesPerSegment = 16;
+		Compression = 0.f;
+		Stretch = 0.f;
+	}
+
+	/** Execute logic for this rig unit */
+	RIGVM_METHOD()
+	virtual void Execute() override;
+
+	UPROPERTY(meta = (Input))
+	TArray<FTransform> Transforms;
+
+	UPROPERTY(meta = (Input))
+	ESplineType SplineMode;
+
+	UPROPERTY(meta = (Input))
+	bool bClosed;
+
+	UPROPERTY(meta = (Input))
+	int32 SamplesPerSegment;
+
+	UPROPERTY(meta = (Input))
+	float Compression;
+
+	UPROPERTY(meta = (Input))
+	float Stretch;
+
+	UPROPERTY(meta = (Output))
+	FControlRigSpline Spline;
+};
+
 /*
  * Set the points of a spline, given a spline and an array of positions
  */
@@ -81,6 +124,30 @@ struct CONTROLRIGSPLINE_API FRigUnit_SetSplinePoints : public FRigUnitMutable
 	FControlRigSpline Spline;
 };
 
+/*
+ * Set the points of a spline, given a spline and an array of positions
+ */
+USTRUCT(meta = (DisplayName = "Set Spline Transforms", Category = "Splines", Varying, NodeColor = "0.0 0.36470600962638855 1.0"))
+struct CONTROLRIGSPLINE_API FRigUnit_SetSplineTransforms : public FRigUnitMutable
+{
+	GENERATED_BODY()
+
+	FRigUnit_SetSplineTransforms()
+	{
+		
+	}
+
+	/** Execute logic for this rig unit */
+	RIGVM_METHOD()
+	virtual void Execute() override;
+
+	UPROPERTY(meta = (Input))
+	TArray<FTransform> Transforms;
+
+	UPROPERTY(meta = (Input, Output))
+	FControlRigSpline Spline;
+};
+
 /*
  * Retrieves the position from a given Spline and U value
  */
@@ -102,7 +169,7 @@ struct CONTROLRIGSPLINE_API FRigUnit_PositionFromControlRigSpline : public FRigU
 	UPROPERTY(meta = (Input))
 	FControlRigSpline Spline;
 
-	UPROPERTY(meta = (Input))
+	UPROPERTY(meta = (Input, ClampMin = "0.0", ClampMax = "1.0"))
 	float U;
 
 	UPROPERTY(meta = (Output))
@@ -166,7 +233,7 @@ struct CONTROLRIGSPLINE_API FRigUnit_TangentFromControlRigSpline : public FRigUn
 	UPROPERTY(meta = (Input))
 	FControlRigSpline Spline;
 
-	UPROPERTY(meta = (Input))
+	UPROPERTY(meta = (Input, ClampMin = "0.0", ClampMax = "1.0"))
 	float U;
 
 	UPROPERTY(meta = (Output))
@@ -229,6 +296,33 @@ struct CONTROLRIGSPLINE_API FRigUnit_GetLengthControlRigSpline : public FRigUnit
 	float Length;
 };
 
+/*
+ * Retrieves the length from a given Splin
+ */
+USTRUCT(meta = (DisplayName = "Get Length At Param Of Spline", Category = "Splines", NodeColor="0.737911 0.099899 0.099899"))
+struct CONTROLRIGSPLINE_API FRigUnit_GetLengthAtParamControlRigSpline : public FRigUnit
+{
+	GENERATED_BODY()
+
+	FRigUnit_GetLengthAtParamControlRigSpline()
+	{
+		Length = 0.f;
+	}
+
+	/** Execute logic for this rig unit */
+	RIGVM_METHOD()
+	virtual void Execute() override;
+
+	UPROPERTY(meta = (Input))
+	FControlRigSpline Spline;
+
+	UPROPERTY(meta = (Input, ClampMin = "0.0", ClampMax = "1.0"))
+	float U;
+
+	UPROPERTY(meta = (Output))
+	float Length;
+};
+
 /**
  * Fits a given chain to a spline curve.
  * Additionally provides rotational control matching the features of the Distribute Rotation node.
@@ -464,6 +558,93 @@ struct CONTROLRIGSPLINE_API FRigUnit_FitChainToSplineCurveItemArray : public FRi
 	FRigUnit_FitChainToCurve_WorkData WorkData;
 };
 
+USTRUCT()
+struct CONTROLRIGSPLINE_API FRigUnit_SplineConstraint_WorkData
+{
+	GENERATED_BODY()
+
+	FRigUnit_SplineConstraint_WorkData()
+	{
+		ChainLength = 0.f;
+	}
+
+	UPROPERTY()
+	float ChainLength;
+
+	UPROPERTY()
+	TArray<FTransform> ItemTransforms;
+
+	UPROPERTY()
+	TArray<float> ItemSegments;
+
+	UPROPERTY()
+	TArray<FCachedRigElement> CachedItems;
+};
+
+
+/**
+ * Fits a given chain to a spline curve.
+ * Additionally provides rotational control matching the features of the Distribute Rotation node.
+ */
+USTRUCT(meta=(DisplayName="Spline Constraint", Category="Splines", Keywords="Fit,Resample,Spline"))
+struct CONTROLRIGSPLINE_API FRigUnit_SplineConstraint : public FRigUnit_HighlevelBaseMutable
+{
+	GENERATED_BODY()
+
+	FRigUnit_SplineConstraint()
+	{
+		Alignment = EControlRigCurveAlignment::Stretched;
+		Minimum = 0.f;
+		Maximum = 1.f;
+		bPropagateToChildren = true;
+	}
+
+	RIGVM_METHOD()
+	virtual void Execute() override;
+
+	/** 
+	 * The items to align
+	 */
+	UPROPERTY(meta = (Input))
+	TArray<FRigElementKey> Items;
+
+	/** 
+	 * The curve to align to
+	 */
+	UPROPERTY(meta = (Input))
+	FControlRigSpline Spline;
+
+	/** 
+	 * Specifies how to align the chain on the curve
+	 */
+	UPROPERTY(meta = (Input, Constant))
+	EControlRigCurveAlignment Alignment;
+
+	/** 
+	 * The minimum U value to use on the curve
+	 */
+	UPROPERTY(meta = (Input, ClampMin = "0.0", ClampMax = "1.0"))
+	float Minimum;
+
+	/** 
+	 * The maximum U value to use on the curve
+	 */
+	UPROPERTY(meta = (Input, ClampMin = "0.0", ClampMax = "1.0"))
+	float Maximum;
+
+	/**
+	 * If set to true all of the global transforms of the children
+	 * of this bone will be recalculated based on their local transforms.
+	 * Note: This is computationally more expensive than turning it off.
+	 */
+	UPROPERTY(meta = (Input, Constant))
+	bool bPropagateToChildren;
+
+	UPROPERTY(transient)
+	FRigUnit_SplineConstraint_WorkData WorkData;
+};
+
+
 /**
  * Fits a given spline curve to a chain.
  */
@@ -522,7 +703,7 @@ struct CONTROLRIGSPLINE_API FRigUnit_FitSplineCurveToChainItemArray : public FRi
 	 * The curve to align
 	 */
 	UPROPERTY(meta = (Input, Output))
-	FControlRigSpline Spline;	
+	FControlRigSpline Spline;
 };
 
 /*
diff --git a/Engine/Source/Editor/GraphEditor/Private/SGraphNodeKnot.cpp b/Engine/Source/Editor/GraphEditor/Private/SGraphNodeKnot.cpp
--- a/Engine/Source/Editor/GraphEditor/Private/SGraphNodeKnot.cpp
+++ b/Engine/Source/Editor/GraphEditor/Private/SGraphNodeKnot.cpp
@@ -246,7 +246,14 @@ void SGraphPinKnot::OnDragEnter(const FGeometry& MyGeometry, const FDragDropEven
 FSlateColor SGraphPinKnot::GetPinColor() const
 {
 	// Make ourselves transparent if we're the input, since we are underneath the output pin and would double-blend looking ugly
-	return (GetPinObj()->Direction == EEdGraphPinDirection::EGPD_Input) ? FLinearColor::Transparent : SGraphPin::GetPinColor();
+	if (UEdGraphPin* EdGraphPin = GetPinObj())
+	{
+		if (EdGraphPin->Direction == EEdGraphPinDirection::EGPD_Input)
+		{
+			return FLinearColor::Transparent;
+		}
+	}
+	return SGraphPin::GetPinColor();
 }
 
 TSharedRef<SWidget>	SGraphPinKnot::GetDefaultValueWidget()
