Control Rig Spline: Fix warning property not initialize in FRigUnit_GetLengthAtParamControlRigSpline

diff --git a/Engine/Plugins/Animation/ControlRigSpline/Source/ControlRigSpline/Public/ControlRigSplineUnits.h b/Engine/Plugins/Animation/ControlRigSpline/Source/ControlRigSpline/Public/ControlRigSplineUnits.h
--- a/Engine/Plugins/Animation/ControlRigSpline/Source/ControlRigSpline/Public/ControlRigSplineUnits.h
+++ b/Engine/Plugins/Animation/ControlRigSpline/Source/ControlRigSpline/Public/ControlRigSplineUnits.h
@@ -306,6 +306,7 @@ struct CONTROLRIGSPLINE_API FRigUnit_GetLengthAtParamControlRigSpline : public F
 
 	FRigUnit_GetLengthAtParamControlRigSpline()
 	{
+		U = 0.f;
 		Length = 0.f;
 	}
 
