Control Rig: Fix static analysis warning

diff --git a/Engine/Plugins/Animation/ControlRigSpline/Source/ControlRigSpline/Private/ControlRigSplineUnits.cpp b/Engine/Plugins/Animation/ControlRigSpline/Source/ControlRigSpline/Private/ControlRigSplineUnits.cpp
--- a/Engine/Plugins/Animation/ControlRigSpline/Source/ControlRigSpline/Private/ControlRigSplineUnits.cpp
+++ b/Engine/Plugins/Animation/ControlRigSpline/Source/ControlRigSpline/Private/ControlRigSplineUnits.cpp
@@ -666,48 +666,6 @@ FRigUnit_FitChainToSplineCurveItemArray_Execute()
 			Hierarchy->SetGlobalTransform(CachedItems[Index], BaseTransform, bPropagateToChildren);
 		}
 	}
-	else
-	{
-		FTransform BaseTransform = FTransform::Identity;
-		FRigElementKey ParentKey = Hierarchy->GetFirstParent(CachedItems[0].GetKey());
-		if (ParentKey.IsValid())
-		{
-			BaseTransform = Hierarchy->GetGlobalTransform(ParentKey);
-		}
-
-		for (int32 Index = 0; Index < CachedItems.Num(); Index++)
-		{
-			ItemLocalTransforms[Index] = Hierarchy->GetLocalTransform(CachedItems[Index]);
-		}
-
-		for (int32 Index = 0; Index < CachedItems.Num(); Index++)
-		{
-			if (ItemRotationA[Index] >= Rotations.Num() ||
-				ItemRotationB[Index] >= Rotations.Num())
-			{
-				continue;
-			}
-
-			FQuat Rotation = Rotations[ItemRotationA[Index]].Rotation;
-			FQuat RotationB = Rotations[ItemRotationB[Index]].Rotation;
-			if (ItemRotationA[Index] != ItemRotationB[Index])
-			{
-				if (ItemRotationT[Index] > 1.f - SMALL_NUMBER)
-				{
-					Rotation = RotationB;
-				}
-				else if (ItemRotationT[Index] > SMALL_NUMBER)
-				{
-					Rotation = FQuat::Slerp(Rotation, RotationB, ItemRotationT[Index]).GetNormalized();
-				}
-			}
-
-			BaseTransform = ItemLocalTransforms[Index] * BaseTransform;
-			Rotation = FQuat::Slerp(BaseTransform.GetRotation(), BaseTransform.GetRotation() * Rotation, FMath::Clamp<float>(Weight, 0.f, 1.f));
-			BaseTransform.SetRotation(Rotation);
-			Hierarchy->SetGlobalTransform(CachedItems[Index], BaseTransform, bPropagateToChildren);
-		}
-	}
 
 	if (ExecuteContext.GetDrawInterface() != nullptr && DebugSettings.bEnabled)
 	{
