Control Rig Spline: Fix interpolation of rotations

diff --git a/Engine/Plugins/Animation/ControlRigSpline/Source/ControlRigSpline/Private/ControlRigSplineTypes.cpp b/Engine/Plugins/Animation/ControlRigSpline/Source/ControlRigSpline/Private/ControlRigSplineTypes.cpp
--- a/Engine/Plugins/Animation/ControlRigSpline/Source/ControlRigSpline/Private/ControlRigSplineTypes.cpp
+++ b/Engine/Plugins/Animation/ControlRigSpline/Source/ControlRigSpline/Private/ControlRigSplineTypes.cpp
@@ -454,7 +454,8 @@ void FControlRigSpline::SetControlTransforms(const TArrayView<const FTransform>&
 		// Cache sample positions of the spline
 		const int32 NumSamples = (ControlPointsCount-1) * SamplesPerSegment;
 		SplineData->SamplesArray.SetNumUninitialized(NumSamples, false);
-			
+
+		// Compute sample positions
 		float U = 0.f;
 		const float DeltaU = 1.f/(NumSamples-1);
 		for (int32 i=0; i<NumSamples; ++i, U+=DeltaU)
@@ -462,9 +463,7 @@ void FControlRigSpline::SetControlTransforms(const TArrayView<const FTransform>&
 			SplineData->SamplesArray[i].SetTranslation(SplineData->Spline->GetPointAtParam(U));
 		}
 
-		FVector UpVector = ControlPointsCount > 0 ? InTransforms[0].GetRotation().GetUpVector() : FVector(0,0,0);
-		UpVector.Normalize();
-
+		// Interpolate sample rotations and scales
 		U = 0.f;
 		for (int32 i=0; i<NumSamples; ++i, U+=DeltaU)
 		{
@@ -484,12 +483,7 @@ void FControlRigSpline::SetControlTransforms(const TArrayView<const FTransform>&
 			const float UNext = NextIndex / (float)(ControlPointsCount-1);
 			const float Interp = (U - UPrev) / (UNext - UPrev);
 			const FQuat InterpRotation = FQuat::Slerp(InTransforms[PrevIndex].GetRotation(), InTransforms[NextIndex].GetRotation(), Interp);
-			
-			//FVector Binormal = FVector::CrossProduct(Tangent, UpVector);
-			FMatrix RotationMatrix = FRotationMatrix::MakeFromXZ(Tangent, InterpRotation.GetUpVector());
-			FQuat RotationQuat = RotationMatrix.ToQuat();
-			RotationQuat.Normalize();
-			SplineData->SamplesArray[i].SetRotation(RotationQuat);
+			SplineData->SamplesArray[i].SetRotation(InterpRotation);
 
 			SplineData->SamplesArray[i].SetScale3D(FMath::Lerp(InTransforms[PrevIndex].GetScale3D(), InTransforms[NextIndex].GetScale3D(), Interp));
 		}
