Control Rig Spline: Add primary and secondary axis to spline constraint and transform from spline

diff --git a/Engine/Plugins/Animation/ControlRigSpline/Source/ControlRigSpline/Private/ControlRigSplineUnits.cpp b/Engine/Plugins/Animation/ControlRigSpline/Source/ControlRigSpline/Private/ControlRigSplineUnits.cpp
--- a/Engine/Plugins/Animation/ControlRigSpline/Source/ControlRigSpline/Private/ControlRigSplineUnits.cpp
+++ b/Engine/Plugins/Animation/ControlRigSpline/Source/ControlRigSpline/Private/ControlRigSplineUnits.cpp
@@ -140,6 +140,23 @@ FRigUnit_TransformFromControlRigSpline_Execute()
 	const float ClampedU = FMath::Clamp<float>(U, 0.f, 1.f);
 	Transform = Spline.TransformAtParam(ClampedU);
 
+	FTransform OriginalTransform = Transform;
+	if (!PrimaryAxis.IsNearlyZero())
+	{
+		FVector Axis = Transform.TransformVectorNoScale(PrimaryAxis);
+		FVector Target1 = Transform.GetRotation().GetForwardVector();
+		FQuat Rotation1 = FQuat::FindBetweenNormals(Axis, Target1);
+		Transform.SetRotation((Rotation1 * Transform.GetRotation()).GetNormalized());
+	}
+
+	if (!SecondaryAxis.IsNearlyZero())
+	{
+		FVector Axis1 = OriginalTransform.TransformVectorNoScale(SecondaryAxis);
+		FVector Axis = Transform.TransformVectorNoScale(Axis1);
+		FVector Target2 = Transform.GetRotation().GetForwardVector();
+		FQuat Rotation2 = FQuat::FindBetweenNormals(Axis, Target2);
+		Transform.SetRotation((Rotation2 * Transform.GetRotation()).GetNormalized());
+	}
 }
 
 FRigUnit_TangentFromControlRigSpline_Execute()
@@ -824,6 +841,36 @@ FRigUnit_SplineConstraint_Execute()
 		}
 	}
 
+	FQuat RotatePrimary = FQuat::Identity;
+	FQuat RotateSecondary = FQuat::Identity;
+	// Find primary and secondary rotations based on primary and secondary axes
+	{
+		const FTransform OriginalTransform = ItemTransforms[0];
+		FTransform FirstTransform = ItemTransforms[0];
+		if (!PrimaryAxis.IsNearlyZero())
+		{
+			FVector Axis = FirstTransform.TransformVectorNoScale(PrimaryAxis);
+			FVector Target1 = FirstTransform.GetRotation().GetForwardVector();
+			RotatePrimary = FQuat::FindBetweenNormals(Axis, Target1);
+			FirstTransform.SetRotation((RotatePrimary * FirstTransform.GetRotation()).GetNormalized());
+		}
+
+		if (!SecondaryAxis.IsNearlyZero())
+		{
+			FVector Axis1 = OriginalTransform.TransformVectorNoScale(SecondaryAxis);
+			FVector Axis = FirstTransform.TransformVectorNoScale(Axis1);
+			FVector Target2 = FirstTransform.GetRotation().GetForwardVector();
+			RotateSecondary = FQuat::FindBetweenNormals(Axis, Target2);
+		}
+	}
+
+	// Apply primary and secondary rotations to item transforms
+	for (int32 Index = 0; Index < ItemTransforms.Num(); Index++)
+	{
+		FTransform& Transform = ItemTransforms[Index];
+		Transform.SetRotation((RotateSecondary * RotatePrimary * Transform.GetRotation()).GetNormalized());
+	}
+	
 	for (int32 Index = 0; Index < CachedItems.Num(); Index++)
 	{
 		Hierarchy->SetGlobalTransform(CachedItems[Index], ItemTransforms[Index], bPropagateToChildren);
diff --git a/Engine/Plugins/Animation/ControlRigSpline/Source/ControlRigSpline/Public/ControlRigSplineUnits.h b/Engine/Plugins/Animation/ControlRigSpline/Source/ControlRigSpline/Public/ControlRigSplineUnits.h
--- a/Engine/Plugins/Animation/ControlRigSpline/Source/ControlRigSpline/Public/ControlRigSplineUnits.h
+++ b/Engine/Plugins/Animation/ControlRigSpline/Source/ControlRigSpline/Public/ControlRigSplineUnits.h
@@ -186,8 +186,6 @@ struct CONTROLRIGSPLINE_API FRigUnit_TransformFromControlRigSpline : public FRig
 
 	FRigUnit_TransformFromControlRigSpline()
 	{
-		UpVector = FVector::UpVector;
-		Roll = 0.f;
 		U = 0.f;
 		Transform = FTransform::Identity;
 	}
@@ -200,13 +198,13 @@ struct CONTROLRIGSPLINE_API FRigUnit_TransformFromControlRigSpline : public FRig
 	FControlRigSpline Spline;
 
 	UPROPERTY(meta = (Input))
-	FVector UpVector;
+	float U;
 
 	UPROPERTY(meta = (Input))
-	float Roll;
+	FVector PrimaryAxis;
 
 	UPROPERTY(meta = (Input))
-	float U;
+	FVector SecondaryAxis;
 
 	UPROPERTY(meta = (Output))
 	FTransform Transform;
@@ -598,6 +596,8 @@ struct CONTROLRIGSPLINE_API FRigUnit_SplineConstraint : public FRigUnit_Highleve
 		Minimum = 0.f;
 		Maximum = 1.f;
 		bPropagateToChildren = true;
+		PrimaryAxis = FVector::ZeroVector;
+		SecondaryAxis = FVector::ZeroVector;
 	}
 
 	RIGVM_METHOD()
@@ -633,6 +633,12 @@ struct CONTROLRIGSPLINE_API FRigUnit_SplineConstraint : public FRigUnit_Highleve
 	UPROPERTY(meta = (Input, ClampMin = "0.0", ClampMax = "1.0"))
 	float Maximum;
 
+	UPROPERTY(meta = (Input))
+	FVector PrimaryAxis;
+
+	UPROPERTY(meta = (Input))
+	FVector SecondaryAxis;
+
 	/**
 	 * If set to true all of the global transforms of the children
 	 * of this bone will be recalculated based on their local transforms.
@@ -672,7 +678,7 @@ struct CONTROLRIGSPLINE_API FRigUnit_FitSplineCurveToChain : public FRigUnit_Hig
 	 * The curve to align
 	 */
 	UPROPERTY(meta = (Input, Output))
-	FControlRigSpline Spline;	
+	FControlRigSpline Spline;
 
 	RIGVM_METHOD()
 	virtual FRigVMStructUpgradeInfo GetUpgradeInfo() const override;
