Control Rig Spline: Initialize Primary and Secondary Axis in unit nodes

diff --git a/Engine/Plugins/Animation/ControlRigSpline/Source/ControlRigSpline/Public/ControlRigSplineUnits.h b/Engine/Plugins/Animation/ControlRigSpline/Source/ControlRigSpline/Public/ControlRigSplineUnits.h
--- a/Engine/Plugins/Animation/ControlRigSpline/Source/ControlRigSpline/Public/ControlRigSplineUnits.h
+++ b/Engine/Plugins/Animation/ControlRigSpline/Source/ControlRigSpline/Public/ControlRigSplineUnits.h
@@ -188,6 +188,8 @@ struct CONTROLRIGSPLINE_API FRigUnit_TransformFromControlRigSpline : public FRig
 	{
 		U = 0.f;
 		Transform = FTransform::Identity;
+		PrimaryAxis = FVector::ZeroVector;
+		SecondaryAxis = FVector::ZeroVector;
 	}
 
 	/** Execute logic for this rig unit */
