Control Rig Spline: Fix incorrect orientation in Tranform from spline

diff --git a/Engine/Plugins/Animation/ControlRigSpline/Source/ControlRigSpline/Private/ControlRigSplineTypes.cpp b/Engine/Plugins/Animation/ControlRigSpline/Source/ControlRigSpline/Private/ControlRigSplineTypes.cpp
--- a/Engine/Plugins/Animation/ControlRigSpline/Source/ControlRigSpline/Private/ControlRigSplineTypes.cpp
+++ b/Engine/Plugins/Animation/ControlRigSpline/Source/ControlRigSpline/Private/ControlRigSplineTypes.cpp
@@ -483,7 +483,8 @@ void FControlRigSpline::SetControlTransforms(const TArrayView<const FTransform>&
 			const float UNext = NextIndex / (float)(ControlPointsCount-1);
 			const float Interp = (U - UPrev) / (UNext - UPrev);
 			const FQuat InterpRotation = FQuat::Slerp(InTransforms[PrevIndex].GetRotation(), InTransforms[NextIndex].GetRotation(), Interp);
-			SplineData->SamplesArray[i].SetRotation(InterpRotation);
+			FQuat CurveRotation = FQuat::FindBetween(FVector::UnitX(), Tangent);
+			SplineData->SamplesArray[i].SetRotation(InterpRotation * CurveRotation);
 
 			SplineData->SamplesArray[i].SetScale3D(FMath::Lerp(InTransforms[PrevIndex].GetScale3D(), InTransforms[NextIndex].GetScale3D(), Interp));
 		}
