Control Rig Spline: Restore the old TransformFromSpline and mark it as deprecated

diff --git a/Engine/Plugins/Animation/ControlRigSpline/Source/ControlRigSpline/Private/ControlRigSplineUnits.cpp b/Engine/Plugins/Animation/ControlRigSpline/Source/ControlRigSpline/Private/ControlRigSplineUnits.cpp
--- a/Engine/Plugins/Animation/ControlRigSpline/Source/ControlRigSpline/Private/ControlRigSplineUnits.cpp
+++ b/Engine/Plugins/Animation/ControlRigSpline/Source/ControlRigSpline/Private/ControlRigSplineUnits.cpp
@@ -131,6 +131,52 @@ FRigUnit_TransformFromControlRigSpline_Execute()
 		return;
 	}
 
+	if (Spline.SplineData->Spline == nullptr)
+	{
+		UE_LOG(LogControlRig, Error, TEXT("Invalid input spline implementation."));
+		return;
+	}
+	
+	FVector UpVectorNormalized = UpVector;
+	UpVectorNormalized.Normalize();
+
+	const float ClampedU = FMath::Clamp<float>(U, 0.f, 1.f);
+	const float ClampedRoll = FMath::Clamp<float>(Roll, -180.f, 180.f);
+
+	FVector Tangent = Spline.TangentAtParam(ClampedU);
+
+	// Check if Tangent can be normalized. If not, keep the same tangent as before.
+	if (!Tangent.Normalize())
+	{
+		Tangent = Transform.ToMatrixNoScale().GetUnitAxis(EAxis::X);
+	}
+	FVector Binormal = FVector::CrossProduct(Tangent, UpVectorNormalized);
+	Binormal = Binormal.RotateAngleAxis(ClampedRoll * ClampedU, Tangent);
+	
+	FMatrix RotationMatrix = FRotationMatrix::MakeFromXZ(Tangent, Binormal);
+
+	Transform.SetFromMatrix(RotationMatrix);
+	Transform.SetTranslation(Spline.PositionAtParam(U));
+}
+
+FRigVMStructUpgradeInfo FRigUnit_TransformFromControlRigSpline::GetUpgradeInfo() const
+{
+	FRigUnit_TransformFromControlRigSpline2 NewNode;
+	NewNode.Spline = Spline; 
+	NewNode.U = U; 
+	NewNode.SecondaryAxis = UpVector; 
+	NewNode.Transform = Transform; 
+
+	return FRigVMStructUpgradeInfo(*this, NewNode);
+}
+
+FRigUnit_TransformFromControlRigSpline2_Execute()
+{
+	if (!Spline.SplineData.IsValid())
+	{
+		return;
+	}
+
 	if (Spline.SplineData->Spline == nullptr)
 	{
 		UE_LOG(LogControlRig, Error, TEXT("Invalid input spline implementation."));
diff --git a/Engine/Plugins/Animation/ControlRigSpline/Source/ControlRigSpline/Public/ControlRigSplineUnits.h b/Engine/Plugins/Animation/ControlRigSpline/Source/ControlRigSpline/Public/ControlRigSplineUnits.h
--- a/Engine/Plugins/Animation/ControlRigSpline/Source/ControlRigSpline/Public/ControlRigSplineUnits.h
+++ b/Engine/Plugins/Animation/ControlRigSpline/Source/ControlRigSpline/Public/ControlRigSplineUnits.h
@@ -179,12 +179,51 @@ struct CONTROLRIGSPLINE_API FRigUnit_PositionFromControlRigSpline : public FRigU
 /*
  * Retrieves the transform from a given Spline and U value based on the given Up Vector and Roll
  */
-USTRUCT(meta = (DisplayName = "Transform From Spline", Category = "Splines", NodeColor="0.737911 0.099899 0.099899"))
+USTRUCT(meta = (DisplayName = "Transform From Spline", Category = "Splines", NodeColor="0.737911 0.099899 0.099899", Deprecated="5.3"))
 struct CONTROLRIGSPLINE_API FRigUnit_TransformFromControlRigSpline : public FRigUnit_ControlRigSplineBase
 {
 	GENERATED_BODY()
 
 	FRigUnit_TransformFromControlRigSpline()
+	{
+		UpVector = FVector::UpVector;
+		Roll = 0.f;
+		U = 0.f;
+		Transform = FTransform::Identity;
+	}
+
+	/** Execute logic for this rig unit */
+	RIGVM_METHOD()
+	virtual void Execute() override;
+
+	UPROPERTY(meta = (Input))
+	FControlRigSpline Spline;
+
+	UPROPERTY(meta = (Input))
+	FVector UpVector;
+
+	UPROPERTY(meta = (Input))
+	float Roll;
+
+	UPROPERTY(meta = (Input))
+	float U;
+
+	UPROPERTY(meta = (Output))
+	FTransform Transform;
+
+	RIGVM_METHOD()
+	virtual FRigVMStructUpgradeInfo GetUpgradeInfo() const override;
+};
+
+/*
+ * Retrieves the transform from a given Spline and U value based on the given Up Vector and Roll
+ */
+USTRUCT(meta = (DisplayName = "Transform From Spline", Category = "Splines", NodeColor="0.737911 0.099899 0.099899"))
+struct CONTROLRIGSPLINE_API FRigUnit_TransformFromControlRigSpline2 : public FRigUnit_ControlRigSplineBase
+{
+	GENERATED_BODY()
+
+	FRigUnit_TransformFromControlRigSpline2()
 	{
 		U = 0.f;
 		Transform = FTransform::Identity;
