Control Rig Spline: Fix rotations so that they always follow the tangent

diff --git a/Engine/Plugins/Animation/ControlRigSpline/Source/ControlRigSpline/Private/ControlRigSplineTypes.cpp b/Engine/Plugins/Animation/ControlRigSpline/Source/ControlRigSpline/Private/ControlRigSplineTypes.cpp
--- a/Engine/Plugins/Animation/ControlRigSpline/Source/ControlRigSpline/Private/ControlRigSplineTypes.cpp
+++ b/Engine/Plugins/Animation/ControlRigSpline/Source/ControlRigSpline/Private/ControlRigSplineTypes.cpp
@@ -467,8 +467,21 @@ void FControlRigSpline::SetControlTransforms(const TArrayView<const FTransform>&
 		U = 0.f;
 		for (int32 i=0; i<NumSamples; ++i, U+=DeltaU)
 		{
+			// This tangent will always be the X axis of the sample transform
 			FVector Tangent = TangentAtParam(U);
 			Tangent.Normalize();
+
+			FVector UpVector = FVector::UnitY();
+			FVector Binormal = Tangent.Cross(UpVector);
+			Binormal.Normalize();
+			UpVector = Binormal.Cross(Tangent);
+			UpVector.Normalize();
+			
+			FMatrix Matrix(Tangent, UpVector, Binormal, FVector::ZeroVector);
+			SplineData->SamplesArray[i].SetRotation(Matrix.ToQuat());
+			SplineData->SamplesArray[i].SetScale3D(FVector(1.f,1.f,1.f));
+
+			FTransform& ParallelTransform = SplineData->SamplesArray[i];
 			
 			int32 PrevIndex = U*(ControlPointsCount-1);
 			int32 NextIndex = PrevIndex+1;
@@ -477,16 +490,24 @@ void FControlRigSpline::SetControlTransforms(const TArrayView<const FTransform>&
 				PrevIndex--;
 				NextIndex--;
 			}
-			
+
 			// Interpolate inside segment
 			const float UPrev = PrevIndex / (float)(ControlPointsCount-1);
 			const float UNext = NextIndex / (float)(ControlPointsCount-1);
 			const float Interp = (U - UPrev) / (UNext - UPrev);
-			const FQuat InterpRotation = FQuat::Slerp(InTransforms[PrevIndex].GetRotation(), InTransforms[NextIndex].GetRotation(), Interp);
-			FQuat CurveRotation = FQuat::FindBetween(FVector::UnitX(), Tangent);
-			SplineData->SamplesArray[i].SetRotation(InterpRotation * CurveRotation);
 
-			SplineData->SamplesArray[i].SetScale3D(FMath::Lerp(InTransforms[PrevIndex].GetScale3D(), InTransforms[NextIndex].GetScale3D(), Interp));
+			const FTransform TransportedTransformPrev = ParallelTransform * InTransforms[PrevIndex];
+			const FTransform TransportedTransformNext = ParallelTransform * InTransforms[NextIndex];
+			const FQuat DiffRotation = FQuat::Slerp(InTransforms[PrevIndex].GetRotation(), InTransforms[NextIndex].GetRotation(), Interp);
+
+			FVector UpVectorRotated = DiffRotation.RotateVector(FVector::UnitY());
+			FVector BinormalRotated = Tangent.Cross(UpVectorRotated);
+			BinormalRotated.Normalize();
+			UpVectorRotated = BinormalRotated.Cross(Tangent);
+			UpVectorRotated.Normalize();
+			FMatrix MatrixRotated(Tangent, UpVectorRotated, BinormalRotated, FVector::ZeroVector);
+			SplineData->SamplesArray[i].SetRotation(MatrixRotated.ToQuat());
+			SplineData->SamplesArray[i].SetScale3D(FVector(1.f,1.f,1.f));
 		}
 
 		// Correct length of samples
diff --git a/Engine/Plugins/Animation/ControlRigSpline/Source/ControlRigSpline/Private/ControlRigSplineUnits.cpp b/Engine/Plugins/Animation/ControlRigSpline/Source/ControlRigSpline/Private/ControlRigSplineUnits.cpp
--- a/Engine/Plugins/Animation/ControlRigSpline/Source/ControlRigSpline/Private/ControlRigSplineUnits.cpp
+++ b/Engine/Plugins/Animation/ControlRigSpline/Source/ControlRigSpline/Private/ControlRigSplineUnits.cpp
@@ -8,6 +8,28 @@
 
 #include UE_INLINE_GENERATED_CPP_BY_NAME(ControlRigSplineUnits)
 
+FTransform ApplyPrimarySecondaryAxis(const FTransform& OriginalTransform, const FVector& PrimaryAxis, const FVector& SecondaryAxis)
+{
+	FTransform Transform = OriginalTransform;
+	if (!PrimaryAxis.IsNearlyZero())
+	{
+		FVector Axis = Transform.TransformVectorNoScale(PrimaryAxis);
+		FVector Target1 = Transform.GetRotation().GetForwardVector();
+		FQuat Rotation1 = FQuat::FindBetweenNormals(Axis, Target1);
+		Transform.SetRotation((Rotation1 * Transform.GetRotation()).GetNormalized());
+	}
+
+	if (!SecondaryAxis.IsNearlyZero())
+	{
+		FVector Axis = Transform.TransformVectorNoScale(SecondaryAxis);
+		FVector Target1 = Transform.GetRotation().GetUpVector();
+		FQuat Rotation1 = FQuat::FindBetweenNormals(Axis, Target1);
+		Transform.SetRotation((Rotation1 * Transform.GetRotation()).GetNormalized());
+	}
+
+	return Transform;
+}
+
 FRigUnit_ControlRigSplineFromPoints_Execute()
 {
 	// reset the spline
@@ -186,23 +208,7 @@ FRigUnit_TransformFromControlRigSpline2_Execute()
 	const float ClampedU = FMath::Clamp<float>(U, 0.f, 1.f);
 	Transform = Spline.TransformAtParam(ClampedU);
 
-	FTransform OriginalTransform = Transform;
-	if (!PrimaryAxis.IsNearlyZero())
-	{
-		FVector Axis = Transform.TransformVectorNoScale(PrimaryAxis);
-		FVector Target1 = Transform.GetRotation().GetForwardVector();
-		FQuat Rotation1 = FQuat::FindBetweenNormals(Axis, Target1);
-		Transform.SetRotation((Rotation1 * Transform.GetRotation()).GetNormalized());
-	}
-
-	if (!SecondaryAxis.IsNearlyZero())
-	{
-		FVector Axis1 = OriginalTransform.TransformVectorNoScale(SecondaryAxis);
-		FVector Axis = Transform.TransformVectorNoScale(Axis1);
-		FVector Target2 = Transform.GetRotation().GetForwardVector();
-		FQuat Rotation2 = FQuat::FindBetweenNormals(Axis, Target2);
-		Transform.SetRotation((Rotation2 * Transform.GetRotation()).GetNormalized());
-	}
+	Transform = ApplyPrimarySecondaryAxis(Transform, PrimaryAxis, SecondaryAxis);
 }
 
 FRigUnit_TangentFromControlRigSpline_Execute()
@@ -887,34 +893,10 @@ FRigUnit_SplineConstraint_Execute()
 		}
 	}
 
-	FQuat RotatePrimary = FQuat::Identity;
-	FQuat RotateSecondary = FQuat::Identity;
-	// Find primary and secondary rotations based on primary and secondary axes
-	{
-		const FTransform OriginalTransform = ItemTransforms[0];
-		FTransform FirstTransform = ItemTransforms[0];
-		if (!PrimaryAxis.IsNearlyZero())
-		{
-			FVector Axis = FirstTransform.TransformVectorNoScale(PrimaryAxis);
-			FVector Target1 = FirstTransform.GetRotation().GetForwardVector();
-			RotatePrimary = FQuat::FindBetweenNormals(Axis, Target1);
-			FirstTransform.SetRotation((RotatePrimary * FirstTransform.GetRotation()).GetNormalized());
-		}
-
-		if (!SecondaryAxis.IsNearlyZero())
-		{
-			FVector Axis1 = OriginalTransform.TransformVectorNoScale(SecondaryAxis);
-			FVector Axis = FirstTransform.TransformVectorNoScale(Axis1);
-			FVector Target2 = FirstTransform.GetRotation().GetForwardVector();
-			RotateSecondary = FQuat::FindBetweenNormals(Axis, Target2);
-		}
-	}
-
-	// Apply primary and secondary rotations to item transforms
 	for (int32 Index = 0; Index < ItemTransforms.Num(); Index++)
 	{
 		FTransform& Transform = ItemTransforms[Index];
-		Transform.SetRotation((RotateSecondary * RotatePrimary * Transform.GetRotation()).GetNormalized());
+		Transform = ApplyPrimarySecondaryAxis(Transform, PrimaryAxis, SecondaryAxis);
 	}
 	
 	for (int32 Index = 0; Index < CachedItems.Num(); Index++)
