Control Rig: Transform from spline with up vector
#rnx

diff --git a/Engine/Plugins/Animation/ControlRigSpline/Source/ControlRigSpline/Private/ControlRigSplineUnits.cpp b/Engine/Plugins/Animation/ControlRigSpline/Source/ControlRigSpline/Private/ControlRigSplineUnits.cpp
--- a/Engine/Plugins/Animation/ControlRigSpline/Source/ControlRigSpline/Private/ControlRigSplineUnits.cpp
+++ b/Engine/Plugins/Animation/ControlRigSpline/Source/ControlRigSpline/Private/ControlRigSplineUnits.cpp
@@ -182,17 +182,6 @@ FRigUnit_TransformFromControlRigSpline_Execute()
 	Transform.SetTranslation(Spline.PositionAtParam(U));
 }
 
-FRigVMStructUpgradeInfo FRigUnit_TransformFromControlRigSpline::GetUpgradeInfo() const
-{
-	FRigUnit_TransformFromControlRigSpline2 NewNode;
-	NewNode.Spline = Spline; 
-	NewNode.U = U; 
-	NewNode.SecondaryAxis = UpVector; 
-	NewNode.Transform = Transform; 
-
-	return FRigVMStructUpgradeInfo(*this, NewNode);
-}
-
 FRigUnit_TransformFromControlRigSpline2_Execute()
 {
 	if (!Spline.SplineData.IsValid())
diff --git a/Engine/Plugins/Animation/ControlRigSpline/Source/ControlRigSpline/Public/ControlRigSplineUnits.h b/Engine/Plugins/Animation/ControlRigSpline/Source/ControlRigSpline/Public/ControlRigSplineUnits.h
--- a/Engine/Plugins/Animation/ControlRigSpline/Source/ControlRigSpline/Public/ControlRigSplineUnits.h
+++ b/Engine/Plugins/Animation/ControlRigSpline/Source/ControlRigSpline/Public/ControlRigSplineUnits.h
@@ -178,7 +178,7 @@ struct CONTROLRIGSPLINE_API FRigUnit_PositionFromControlRigSpline : public FRigU
 /*
  * Retrieves the transform from a given Spline and U value based on the given Up Vector and Roll
  */
-USTRUCT(meta = (DisplayName = "Transform From Spline", Category = "Splines", NodeColor="0.737911 0.099899 0.099899", Deprecated="5.3"))
+USTRUCT(meta = (DisplayName = "Transform From Spline (with UpVector)", Category = "Splines", NodeColor="0.737911 0.099899 0.099899"))
 struct CONTROLRIGSPLINE_API FRigUnit_TransformFromControlRigSpline : public FRigUnit_ControlRigSplineBase
 {
 	GENERATED_BODY()
@@ -209,9 +209,6 @@ struct CONTROLRIGSPLINE_API FRigUnit_TransformFromControlRigSpline : public FRig
 
 	UPROPERTY(meta = (Output))
 	FTransform Transform;
-
-	RIGVM_METHOD()
-	virtual FRigVMStructUpgradeInfo GetUpgradeInfo() const override;
 };
 
 /*
