ControlRigSpline: Guard against transform indices on closed splines

#rnx

diff --git a/Engine/Plugins/Animation/ControlRigSpline/Source/ControlRigSpline/Private/ControlRigSplineTypes.cpp b/Engine/Plugins/Animation/ControlRigSpline/Source/ControlRigSpline/Private/ControlRigSplineTypes.cpp
--- a/Engine/Plugins/Animation/ControlRigSpline/Source/ControlRigSpline/Private/ControlRigSplineTypes.cpp
+++ b/Engine/Plugins/Animation/ControlRigSpline/Source/ControlRigSpline/Private/ControlRigSplineTypes.cpp
@@ -296,7 +296,15 @@ TArray<FTransform> FControlRigSplineImpl::GetControlTransformsWithoutDuplicates(
 	Transforms.Reserve(ControlIndices.Num());
 	for (const uint16& Index : ControlIndices)
 	{
-		Transforms.Add(ControlTransforms[Index]);
+		int32 WrappedIndex = int32(Index);
+
+		// for closed splines we may ask for a control index above the index count,
+		// so we need to wrap around along the spline and start over from the front.
+		if(!ControlTransforms.IsValidIndex(WrappedIndex))
+		{
+			WrappedIndex = WrappedIndex % ControlTransforms.Num();
+		}
+		Transforms.Add(ControlTransforms[WrappedIndex]);
 	}
 	return Transforms;
 }
