Control Rig Spline: Fix plugin dependency missing for RigVM
#rnx

diff --git a/Engine/Plugins/Animation/ControlRigSpline/ControlRigSpline.uplugin b/Engine/Plugins/Animation/ControlRigSpline/ControlRigSpline.uplugin
--- a/Engine/Plugins/Animation/ControlRigSpline/ControlRigSpline.uplugin
+++ b/Engine/Plugins/Animation/ControlRigSpline/ControlRigSpline.uplugin
@@ -19,6 +19,10 @@
 		{
 			"Name": "ControlRig",
 			"Enabled": true
+		},
+		{
+			"Name": "RigVM",
+			"Enabled": true
 		}
 	],
 	"Modules": [
