[Control Rig Spline] added tangent interpolation between spline samples.

diff --git a/Engine/Plugins/Animation/ControlRigSpline/Source/ControlRigSpline/Private/ControlRigSplineTypes.cpp b/Engine/Plugins/Animation/ControlRigSpline/Source/ControlRigSpline/Private/ControlRigSplineTypes.cpp
--- a/Engine/Plugins/Animation/ControlRigSpline/Source/ControlRigSpline/Private/ControlRigSplineTypes.cpp
+++ b/Engine/Plugins/Animation/ControlRigSpline/Source/ControlRigSpline/Private/ControlRigSplineTypes.cpp
@@ -2,6 +2,8 @@
 
 #include "ControlRigSplineTypes.h"
 
+#include "RigVMFunctions/Math/RigVMMathLibrary.h"
+
 #include UE_INLINE_GENERATED_CPP_BY_NAME(ControlRigSplineTypes)
 
 ControlRigBaseSpline::ControlRigBaseSpline(const TArrayView<const FVector>& InControlPoints, const uint16 InDegree, bool bInClosed)
@@ -621,11 +623,31 @@ FVector FControlRigSpline::TangentAtParam(const float InParam) const
 	{
 		return FVector();
 	}
+	
+	
 
 	const float ClampedU = FMath::Clamp<float>(InParam, 0.f, 1.f);
 
-	const int32 IndexPrev = ClampedU * (SplineData->SamplesArray.Num()-2);
-	return SplineData->SamplesArray[IndexPrev+1].GetTranslation() - SplineData->SamplesArray[IndexPrev].GetTranslation();
+	const int32 LastIndex = SplineData->SamplesArray.Num() - 1;
+	const float fIndexPrev = ClampedU * LastIndex;	
+	const int32 IndexPrev = FMath::Floor<int32>(fIndexPrev);
+	const float ULocal = fIndexPrev - IndexPrev;
+	const int32 IndexNext = (IndexPrev < LastIndex) ? IndexPrev + 1 : IndexPrev;
+	
+	const FVector TangentStart = SplineData->SamplesArray[IndexNext].GetTranslation() - SplineData->SamplesArray[IndexPrev].GetTranslation();
+	if (IndexNext + 1 < SplineData->SamplesArray.Num())
+	{
+		const FVector TangentEnd = SplineData->SamplesArray[IndexNext+1].GetTranslation() - SplineData->SamplesArray[IndexNext].GetTranslation();
+
+		const FQuat FullQuat = FRigVMMathLibrary::FindQuatBetweenVectors(TangentStart, TangentEnd);
+		const FQuat PartialQuat = FQuat::Slerp(FQuat::Identity, FullQuat, ULocal);
+
+		const FVector Result = PartialQuat.RotateVector(TangentStart);
+
+		return Result;
+	}
+	
+	return TangentStart;
 }
 
 float FControlRigSpline::LengthAtParam(const float InParam) const
