Control Rig: Apply strict documentation policy to ... nodes

#rnx

diff --git a/Engine/Plugins/Animation/ControlRig/Source/ControlRig/Public/Units/Debug/RigUnit_DebugBase.h b/Engine/Plugins/Animation/ControlRig/Source/ControlRig/Public/Units/Debug/RigUnit_DebugBase.h
--- a/Engine/Plugins/Animation/ControlRig/Source/ControlRig/Public/Units/Debug/RigUnit_DebugBase.h
+++ b/Engine/Plugins/Animation/ControlRig/Source/ControlRig/Public/Units/Debug/RigUnit_DebugBase.h
@@ -6,7 +6,10 @@
 #include "RigVMCore/RigVMDebugDrawSettings.h"
 #include "RigUnit_DebugBase.generated.h"
 
-USTRUCT(meta=(Abstract, Category="Debug", NodeColor = "0.83077 0.846873 0.049707"))
+/*
+ * Base class for all pure debug nodes
+ */
+USTRUCT(meta=(Abstract, Category="Debug", NodeColor = "0.83077 0.846873 0.049707", DocumentationPolicy = "Strict"))
 struct FRigUnit_DebugBase : public FRigUnit
 {
 	GENERATED_BODY()
@@ -15,7 +18,10 @@ struct FRigUnit_DebugBase : public FRigUnit
 	FRigVMDebugDrawSettings DebugDrawSettings;
 };
 
-USTRUCT(meta=(Abstract, Category="Debug", NodeColor = "0.83077 0.846873 0.049707"))
+/*
+ * Base class for all mutable debug nodes
+ */
+USTRUCT(meta=(Abstract, Category="Debug", NodeColor = "0.83077 0.846873 0.049707", DocumentationPolicy = "Strict"))
 struct FRigUnit_DebugBaseMutable : public FRigUnitMutable
 {
 	GENERATED_BODY()
diff --git a/Engine/Plugins/Animation/ControlRig/Source/ControlRig/Public/Units/Execution/RigUnit_Collection.h b/Engine/Plugins/Animation/ControlRig/Source/ControlRig/Public/Units/Execution/RigUnit_Collection.h
--- a/Engine/Plugins/Animation/ControlRig/Source/ControlRig/Public/Units/Execution/RigUnit_Collection.h
+++ b/Engine/Plugins/Animation/ControlRig/Source/ControlRig/Public/Units/Execution/RigUnit_Collection.h
@@ -8,13 +8,19 @@
 
 #define UE_API CONTROLRIG_API
 
-USTRUCT(meta = (Abstract, NodeColor = "0.4627450108528137 1.0 0.3294120132923126", Category = "Items"))
+/*
+ * Base class for all pure collection nodes
+ */
+USTRUCT(meta = (Abstract, NodeColor = "0.4627450108528137 1.0 0.3294120132923126", Category = "Items", DocumentationPolicy = "Strict"))
 struct FRigUnit_CollectionBase : public FRigUnit
 {
 	GENERATED_BODY()
 };
 
-USTRUCT(meta = (Abstract, NodeColor = "0.4627450108528137 1.0 0.3294120132923126", Category = "Items"))
+/*
+ * Base class for all mutable collection nodes
+ */
+USTRUCT(meta = (Abstract, NodeColor = "0.4627450108528137 1.0 0.3294120132923126", Category = "Items", DocumentationPolicy = "Strict"))
 struct FRigUnit_CollectionBaseMutable : public FRigUnitMutable
 {
 	GENERATED_BODY()
@@ -137,12 +143,15 @@ struct FRigUnit_CollectionNameSearchArray : public FRigUnit_CollectionBase
 	RIGVM_METHOD()
 	UE_API virtual void Execute() override;
 
+	// The partial name to look for
 	UPROPERTY(meta = (Input))
 	FName PartialName;
 
+	// The type of elements to look for
 	UPROPERTY(meta = (Input))
 	ERigElementType TypeToSearch;
 
+	// The resulting array of elements matching the filters
 	UPROPERTY(meta = (Output))
 	TArray<FRigElementKey> Items;
 };
@@ -207,12 +216,15 @@ struct FRigUnit_CollectionChildrenArray : public FRigUnit_CollectionBase
 	RIGVM_METHOD()
 	UE_API virtual void Execute() override;
 
+	// The parent to search the children for
 	UPROPERTY(meta = (Input, ExpandByDefault))
 	FRigElementKey Parent;
 
+	// If True the parent will be included in the list of results
 	UPROPERTY(meta = (Input))
 	bool bIncludeParent;
 
+	// If True children of children will be included
 	UPROPERTY(meta = (Input))
 	bool bRecursive;
 
@@ -221,9 +233,11 @@ struct FRigUnit_CollectionChildrenArray : public FRigUnit_CollectionBase
 	UPROPERTY(meta = (Input))
 	bool bDefaultChildren;
 
+	// The type of children to look for
 	UPROPERTY(meta = (Input))
 	ERigElementType TypeToSearch;
 
+	// The resulting array of children
 	UPROPERTY(meta = (Output))
 	TArray<FRigElementKey> Items;
 };
@@ -244,9 +258,11 @@ struct FRigUnit_CollectionGetAll : public FRigUnit_CollectionBase
 	RIGVM_METHOD()
 	UE_API virtual void Execute() override;
 
+	// The type of elements to look for
 	UPROPERTY(meta = (Input))
 	ERigElementType TypeToSearch;
 
+	// The resulting array of elements
 	UPROPERTY(meta = (Output))
 	TArray<FRigElementKey> Items;
 };
@@ -309,21 +325,27 @@ struct FRigUnit_CollectionReplaceItemsArray : public FRigUnit_CollectionBase
 	RIGVM_METHOD()
 	UE_API virtual void Execute() override;
 
+	// The input array of items to replace
 	UPROPERTY(meta = (Input))
 	TArray<FRigElementKey> Items;
 
+	// The old name to look for
 	UPROPERTY(meta = (Input))
 	FName Old;
 
+	// The new name to replace the old name with
 	UPROPERTY(meta = (Input))
 	FName New;
 
+	// If True items that don't exist (after the replace operation) will be removed from the list of results
 	UPROPERTY(meta = (Input))
 	bool RemoveInvalidItems;
 
+	// If True the list may contain duplicates
 	UPROPERTY(meta = (Input))
 	bool bAllowDuplicates;
 
+	// The resulting array of items
 	UPROPERTY(meta = (Output))
 	TArray<FRigElementKey> Result;
 };
@@ -345,12 +367,15 @@ struct FRigUnit_CollectionItems : public FRigUnit_CollectionBase
 	RIGVM_METHOD()
 	UE_API virtual void Execute() override;
 
+	// The input array of items
 	UPROPERTY(meta = (Input, ExpandByDefault))
 	TArray<FRigElementKey> Items;
 
+	// If True the collection may contain duplicates 
 	UPROPERTY(meta = (Input))
 	bool bAllowDuplicates;
 
+	// The resulting collection
 	UPROPERTY(meta = (Output))
 	FRigElementKeyCollection Collection;
 };
@@ -370,9 +395,11 @@ struct FRigUnit_CollectionGetItems : public FRigUnit_CollectionBase
 	RIGVM_METHOD()
 	UE_API virtual void Execute() override;
 
+	// The input collection to retrieve the items from
 	UPROPERTY(meta = (Input))
 	FRigElementKeyCollection Collection;
 
+	// The resulting items in the collection 
 	UPROPERTY(meta = (Output))
 	TArray<FRigElementKey> Items;
 };
@@ -395,9 +422,11 @@ struct FRigUnit_CollectionGetParentIndices : public FRigUnit_CollectionBase
 	RIGVM_METHOD()
 	UE_API virtual void Execute() override;
 
+	// The collection to retrieve the parent indices from
 	UPROPERTY(meta = (Input))
 	FRigElementKeyCollection Collection;
 
+	// The resulting parent indices
 	UPROPERTY(meta = (Output))
 	TArray<int32> ParentIndices;
 
@@ -423,9 +452,11 @@ struct FRigUnit_CollectionGetParentIndicesItemArray : public FRigUnit_Collection
 	RIGVM_METHOD()
 	UE_API virtual void Execute() override;
 
+	// The input array of items to retrieve the parent indices for
 	UPROPERTY(meta = (Input))
 	TArray<FRigElementKey> Items;
 
+	// The resulting array of parent indices
 	UPROPERTY(meta = (Output))
 	TArray<int32> ParentIndices;
 };
diff --git a/Engine/Plugins/Animation/ControlRig/Source/ControlRig/Public/Units/Execution/RigUnit_DynamicHierarchy.h b/Engine/Plugins/Animation/ControlRig/Source/ControlRig/Public/Units/Execution/RigUnit_DynamicHierarchy.h
--- a/Engine/Plugins/Animation/ControlRig/Source/ControlRig/Public/Units/Execution/RigUnit_DynamicHierarchy.h
+++ b/Engine/Plugins/Animation/ControlRig/Source/ControlRig/Public/Units/Execution/RigUnit_DynamicHierarchy.h
@@ -8,7 +8,10 @@
 
 #define UE_API CONTROLRIG_API
 
-USTRUCT(meta = (Abstract, NodeColor="0.262745, 0.8, 0, 0.229412", Category = "DynamicHierarchy"))
+/*
+ * The base class for all pure dynamic hierarchy nodes
+ */
+USTRUCT(meta = (Abstract, NodeColor="0.262745, 0.8, 0, 0.229412", Category = "DynamicHierarchy", DocumentationPolicy = "Strict"))
 struct FRigUnit_DynamicHierarchyBase : public FRigUnit
 {
 	GENERATED_BODY()
@@ -19,7 +22,10 @@ struct FRigUnit_DynamicHierarchyBase : public FRigUnit
 		FString* OutErrorMessage = nullptr);
 };
 
-USTRUCT(meta = (Abstract, NodeColor="0.262745, 0.8, 0, 0.229412", Category = "DynamicHierarchy"))
+/*
+ * The base class for all mutable dynamic hierarchy nodes
+ */
+USTRUCT(meta = (Abstract, NodeColor="0.262745, 0.8, 0, 0.229412", Category = "DynamicHierarchy", DocumentationPolicy = "Strict"))
 struct FRigUnit_DynamicHierarchyBaseMutable : public FRigUnitMutable
 {
 	GENERATED_BODY()
@@ -343,18 +349,23 @@ struct FRigUnit_HierarchyImportFromSkeleton : public FRigUnit_DynamicHierarchyBa
 	RIGVM_METHOD()
 	virtual void Execute() override;
 
+	// The namespace to use when importing bones and curves
 	UPROPERTY(meta = (Input))
 	FName NameSpace;
 
+	// If True the curves will be imported
 	UPROPERTY(meta = (Input))
 	bool bIncludeCurves;
 
+	// If true the mesh sockets will be imported
 	UPROPERTY(meta = (Input))
 	bool bIncludeMeshSockets;
 
+	// If true virtual bones will be imported
 	UPROPERTY(meta = (Input))
 	bool bIncludeVirtualBones;
 
+	// The resulting imported items
 	UPROPERTY(meta = (Output))
 	TArray<FRigElementKey> Items;
 };
@@ -390,6 +401,9 @@ struct FRigUnit_HierarchyRemoveElement : public FRigUnit_DynamicHierarchyBaseMut
 	bool bSuccess;
 };
 
+/*
+ * The base class for all spawning nodes for the dynamic hierarchy set of nodes
+ */
 USTRUCT(meta=(Abstract))
 struct FRigUnit_HierarchyAddElement : public FRigUnit_DynamicHierarchyBaseMutable
 {
@@ -1256,7 +1270,7 @@ struct FRigUnit_HierarchyAddAnimationChannelBool : public FRigUnit_HierarchyAddE
 	bool InitialValue;
 
 	/*
-	 * The initial value of the new animation channel
+	 * The minimum value of the new animation channel
 	 */
 	UPROPERTY(meta = (Input))
 	bool MinimumValue;
@@ -1314,7 +1328,7 @@ struct FRigUnit_HierarchyAddAnimationChannelFloat : public FRigUnit_HierarchyAdd
 	float InitialValue;
 
 	/*
-	 * The initial value of the new animation channel
+	 * The minimum value of the new animation channel
 	 */
 	UPROPERTY(meta = (Input))
 	float MinimumValue;
@@ -1359,7 +1373,7 @@ struct FRigUnit_HierarchyAddAnimationChannelScaleFloat : public FRigUnit_Hierarc
 	float InitialValue;
 
 	/*
-	 * The initial value of the new animation channel
+	 * The minimum value of the new animation channel
 	 */
 	UPROPERTY(meta = (Input))
 	float MinimumValue;
@@ -1404,7 +1418,7 @@ struct FRigUnit_HierarchyAddAnimationChannelInteger : public FRigUnit_HierarchyA
 	int32 InitialValue;
 
 	/*
-	 * The initial value of the new animation channel
+	 * The minimum value of the new animation channel
 	 */
 	UPROPERTY(meta = (Input))
 	int32 MinimumValue;
@@ -1472,7 +1486,7 @@ struct FRigUnit_HierarchyAddAnimationChannelVector2D : public FRigUnit_Hierarchy
 	FVector2D InitialValue;
 
 	/*
-	 * The initial value of the new animation channel
+	 * The minimum value of the new animation channel
 	 */
 	UPROPERTY(meta = (Input))
 	FVector2D MinimumValue;
@@ -1538,7 +1552,7 @@ struct FRigUnit_HierarchyAddAnimationChannelVector : public FRigUnit_HierarchyAd
 	FVector InitialValue;
 
 	/*
-	 * The initial value of the new animation channel
+	 * The minimum value of the new animation channel
 	 */
 	UPROPERTY(meta = (Input))
 	FVector MinimumValue;
@@ -1583,7 +1597,7 @@ struct FRigUnit_HierarchyAddAnimationChannelScaleVector : public FRigUnit_Hierar
 	FVector InitialValue;
 
 	/*
-	 * The initial value of the new animation channel
+	 * The minimum value of the new animation channel
 	 */
 	UPROPERTY(meta = (Input))
 	FVector MinimumValue;
@@ -1650,7 +1664,7 @@ struct FRigUnit_HierarchyAddAnimationChannelRotator : public FRigUnit_HierarchyA
 	FRotator InitialValue;
 
 	/*
-	 * The initial value of the new animation channel
+	 * The minimum value of the new animation channel
 	 */
 	UPROPERTY(meta = (Input))
 	FRotator MinimumValue;
diff --git a/Engine/Plugins/Animation/ControlRig/Source/ControlRig/Public/Units/Execution/RigUnit_Hierarchy.h b/Engine/Plugins/Animation/ControlRig/Source/ControlRig/Public/Units/Execution/RigUnit_Hierarchy.h
--- a/Engine/Plugins/Animation/ControlRig/Source/ControlRig/Public/Units/Execution/RigUnit_Hierarchy.h
+++ b/Engine/Plugins/Animation/ControlRig/Source/ControlRig/Public/Units/Execution/RigUnit_Hierarchy.h
@@ -8,13 +8,19 @@
 
 #define UE_API CONTROLRIG_API
 
-USTRUCT(meta = (Abstract, NodeColor="0.462745, 1,0, 0.329412", Category = "Hierarchy"))
+/*
+ * Base class for all pure hierarchy nodes
+ */
+USTRUCT(meta = (Abstract, NodeColor="0.462745, 1,0, 0.329412", Category = "Hierarchy", DocumentationPolicy = "Strict"))
 struct FRigUnit_HierarchyBase : public FRigUnit
 {
 	GENERATED_BODY()
 };
 
-USTRUCT(meta = (Abstract, NodeColor="0.462745, 1,0, 0.329412", Category = "Hierarchy"))
+/*
+ * Base class for all mutable hierarchy nodes
+ */
+USTRUCT(meta = (Abstract, NodeColor="0.462745, 1,0, 0.329412", Category = "Hierarchy", DocumentationPolicy = "Strict"))
 struct FRigUnit_HierarchyBaseMutable : public FRigUnitMutable
 {
 	GENERATED_BODY()
@@ -38,13 +44,15 @@ struct FRigUnit_HierarchyGetParent : public FRigUnit_HierarchyBase
 	RIGVM_METHOD()
 	UE_API virtual void Execute() override;
 
+	// The child to look up the parent for
 	UPROPERTY(meta = (Input, ExpandByDefault))
 	FRigElementKey Child;
 
-	/** When true, it will return the default parent, regardless of whether the parent incluences the element or not  */
+	/** When true, it will return the default parent, regardless of whether the parent influences the element or not  */
 	UPROPERTY(meta = (Input))
 	bool bDefaultParent;
 
+	// The resulting parent of the input child
 	UPROPERTY(meta = (Output))
 	FRigElementKey Parent;
 
@@ -77,15 +85,19 @@ struct FRigUnit_HierarchyGetParents : public FRigUnit_HierarchyBase
 	RIGVM_METHOD()
 	UE_API virtual void Execute() override;
 
+	// The child to look up all parents for
 	UPROPERTY(meta = (Input, ExpandByDefault))
 	FRigElementKey Child;
 
+	// If True the child will be included in the list of results
 	UPROPERTY(meta = (Input))
 	bool bIncludeChild;
 
+	// If True the parents list will be reversed
 	UPROPERTY(meta = (Input))
 	bool bReverse;
 
+	// The resulting array of parents
 	UPROPERTY(meta = (Output))
 	FRigElementKeyCollection Parents;
 
@@ -122,18 +134,23 @@ struct FRigUnit_HierarchyGetParentsItemArray : public FRigUnit_HierarchyBase
 	RIGVM_METHOD()
 	UE_API virtual void Execute() override;
 
+	// The child to look up the parents for
 	UPROPERTY(meta = (Input, ExpandByDefault))
 	FRigElementKey Child;
 
+	// If True the child will be included in the list of results
 	UPROPERTY(meta = (Input))
 	bool bIncludeChild;
 
+	// If True the list of results will be reversed
 	UPROPERTY(meta = (Input))
 	bool bReverse;
 
+	/** When true, it will return the default parent, regardless of whether the parent influences the element or not  */
 	UPROPERTY(meta = (Input))
 	bool bDefaultParent;
 
+	// The resulting array of parents
 	UPROPERTY(meta = (Output))
 	TArray<FRigElementKey> Parents;
 
@@ -251,9 +268,11 @@ struct FRigUnit_HierarchyGetSiblingsItemArray : public FRigUnit_HierarchyBase
 	RIGVM_METHOD()
 	UE_API virtual void Execute() override;
 
+	// The item to find the siblings for
 	UPROPERTY(meta = (Input, ExpandByDefault))
 	FRigElementKey Item;
 
+	// If True the input item will be included in the list of results
 	UPROPERTY(meta = (Input))
 	bool bIncludeItem;
 
@@ -262,6 +281,7 @@ struct FRigUnit_HierarchyGetSiblingsItemArray : public FRigUnit_HierarchyBase
 	UPROPERTY(meta = (Input))
 	bool bDefaultSiblings;
 
+	// The resulting array of siblings of the input item
 	UPROPERTY(meta = (Output))
 	TArray<FRigElementKey> Siblings;
 
@@ -294,21 +314,27 @@ struct FRigUnit_HierarchyGetChainItemArray : public FRigUnit_HierarchyBase
 	RIGVM_METHOD()
 	UE_API virtual void Execute() override;
 
+	// The first element of the chain
 	UPROPERTY(meta = (Input, ExpandByDefault))
 	FRigElementKey Start;
 
+	// The last element of the chain
 	UPROPERTY(meta = (Input, ExpandByDefault))
 	FRigElementKey End;
 
+	// If True the list of results will include the first element
 	UPROPERTY(meta = (Input))
 	bool bIncludeStart;
 
+	// If True the list of results will include the last element
 	UPROPERTY(meta = (Input))
 	bool bIncludeEnd;
 
+	// If True the list of results will be reversed
 	UPROPERTY(meta = (Input))
 	bool bReverse;
-	
+
+	// The resulting array of elements on the chain
 	UPROPERTY(meta = (Output))
 	TArray<FRigElementKey> Chain;
 
@@ -379,9 +405,11 @@ struct FRigUnit_HierarchyGetPoseItemArray : public FRigUnit_HierarchyBase
 	RIGVM_METHOD()
 	UE_API virtual void Execute() override;
 
+	// If True the pose will use the initial transforms
 	UPROPERTY(meta = (Input))
 	bool Initial;
 
+	// The element type to retrieve the pose for 
 	UPROPERTY(meta = (Input))
 	ERigElementType ElementType;
 
@@ -389,6 +417,7 @@ struct FRigUnit_HierarchyGetPoseItemArray : public FRigUnit_HierarchyBase
 	UPROPERTY(meta = (Input))
 	TArray<FRigElementKey> ItemsToGet;
 
+	// The resulting pose
 	UPROPERTY(meta = (Output))
 	FRigPose Pose;
 };
@@ -453,12 +482,15 @@ struct FRigUnit_HierarchySetPoseItemArray : public FRigUnit_HierarchyBaseMutable
 	RIGVM_METHOD()
 	UE_API virtual void Execute() override;
 
+	// The pose to apply to the hierarchy
 	UPROPERTY(meta = (Input))
 	FRigPose Pose;
 
+	// The types of element to apply the pose for
 	UPROPERTY(meta = (Input))
 	ERigElementType ElementType;
 
+	// The space to use to apply the pose in
 	UPROPERTY(meta = (Input))
 	ERigVMTransformSpace Space;
 
@@ -466,6 +498,7 @@ struct FRigUnit_HierarchySetPoseItemArray : public FRigUnit_HierarchyBaseMutable
 	UPROPERTY(meta = (Input))
 	TArray<FRigElementKey> ItemsToSet;
 
+	// The weight to use when applying the transforms
 	UPROPERTY(meta = (Input))
 	float Weight;
 };
@@ -487,9 +520,11 @@ struct FRigUnit_PoseIsEmpty : public FRigUnit_HierarchyBase
 	RIGVM_METHOD()
 	UE_API virtual void Execute() override;
 
+	// The pose to check
 	UPROPERTY(meta = (Input))
 	FRigPose Pose;
 
+	// True if the pose is empty
 	UPROPERTY(meta = (Output))
 	bool IsEmpty;
 };
@@ -542,12 +577,15 @@ struct FRigUnit_PoseGetItemsItemArray : public FRigUnit_HierarchyBase
 	RIGVM_METHOD()
 	UE_API virtual void Execute() override;
 
+	// The pose to retrieve the items from
 	UPROPERTY(meta = (Input))
 	FRigPose Pose;
 
+	// The type of items to retrieve
 	UPROPERTY(meta = (Input))
 	ERigElementType ElementType;
 
+	// The resulting items in the pose
 	UPROPERTY(meta = (Output))
 	TArray<FRigElementKey> Items;
 };
@@ -574,9 +612,11 @@ struct FRigUnit_PoseGetDelta : public FRigUnit_HierarchyBase
 	RIGVM_METHOD()
 	UE_API virtual void Execute() override;
 
+	// The first pose to compare
 	UPROPERTY(meta = (Input))
 	FRigPose PoseA;
 
+	// The second pose to compare
 	UPROPERTY(meta = (Input))
 	FRigPose PoseB;
 
@@ -596,6 +636,7 @@ struct FRigUnit_PoseGetDelta : public FRigUnit_HierarchyBase
 	UPROPERTY(meta = (Input))
 	float CurveThreshold;
 
+	// The type of elements to compare
 	UPROPERTY(meta = (Input))
 	ERigElementType ElementType;
 	
@@ -607,9 +648,11 @@ struct FRigUnit_PoseGetDelta : public FRigUnit_HierarchyBase
 	UPROPERTY(meta = (Input))
 	FRigElementKeyCollection ItemsToCompare;
 
+	// True if the poses A and B are equal
 	UPROPERTY(meta = (Output))
 	bool PosesAreEqual;
 
+	// An array of items with a difference in them between poses A and B
 	UPROPERTY(meta = (Output))
 	FRigElementKeyCollection ItemsWithDelta;
 
@@ -657,9 +700,11 @@ struct FRigUnit_PoseGetTransform : public FRigUnit_HierarchyBase
 	RIGVM_METHOD()
 	UE_API virtual void Execute() override;
 
+	// The pose to retrieve the transform from
 	UPROPERTY(meta = (Input))
 	FRigPose Pose;
 
+	// The item to retrieve the transform for
 	UPROPERTY(meta = (Input))
 	FRigElementKey Item;
 
@@ -669,12 +714,15 @@ struct FRigUnit_PoseGetTransform : public FRigUnit_HierarchyBase
 	UPROPERTY(meta = (Input))
 	ERigVMTransformSpace Space;
 
+	// True if the transform retrieved is valid
 	UPROPERTY(meta = (Output))
 	bool Valid;
-	
+
+	// The resulting transform
 	UPROPERTY(meta = (Output))
 	FTransform Transform;
 
+	// The value of the curve (in case the pose entry was a curve)
 	UPROPERTY(meta = (Output))
 	float CurveValue;
 
@@ -704,6 +752,7 @@ struct FRigUnit_PoseGetTransformArray : public FRigUnit_HierarchyBase
 	RIGVM_METHOD()
 	UE_API virtual void Execute() override;
 
+	// The pose to retrieve the transform from
 	UPROPERTY(meta = (Input))
 	FRigPose Pose;
 
@@ -713,9 +762,11 @@ struct FRigUnit_PoseGetTransformArray : public FRigUnit_HierarchyBase
 	UPROPERTY(meta = (Input))
 	ERigVMTransformSpace Space;
 
+	// True if the transforms has been retrieved successfully
 	UPROPERTY(meta = (Output))
 	bool Valid;
-	
+
+	// The resulting array of transforms
 	UPROPERTY(meta = (Output))
 	TArray<FTransform> Transforms;
 };
@@ -740,15 +791,19 @@ struct FRigUnit_PoseGetCurve : public FRigUnit_HierarchyBase
 	RIGVM_METHOD()
 	UE_API virtual void Execute() override;
 
+	// The pose to retrieve the curve value from
 	UPROPERTY(meta = (Input))
 	FRigPose Pose;
 
+	// The name of the curve to retrieve
 	UPROPERTY(meta = (Input, CustomWidget = "CurveName"))
 	FName Curve;
 
+	// True if the curve value has been retrieved successfully
 	UPROPERTY(meta = (Output))
 	bool Valid;
-	
+
+	// The resulting curve value
 	UPROPERTY(meta = (Output))
 	float CurveValue;
 
@@ -794,24 +849,31 @@ struct FRigUnit_PoseLoop : public FRigUnit_HierarchyBaseMutable
 	UPROPERTY(meta = (Singleton))
 	FName BlockToRun;
 
+	// The pose to loop over
 	UPROPERTY(meta = (Input))
 	FRigPose Pose;
 
+	// The item for the current pose entry
 	UPROPERTY(meta = (Singleton, Output))
 	FRigElementKey Item;
 
+	// The global transform for the current pose entry
 	UPROPERTY(meta = (Singleton, Output))
 	FTransform GlobalTransform;
 
+	// The local transform for the current pose entry
 	UPROPERTY(meta = (Singleton, Output))
 	FTransform LocalTransform;
 
+	// The curve value for the current pose entry
 	UPROPERTY(meta = (Singleton, Output))
 	float CurveValue;
 
+	// The index of the current pose entry
 	UPROPERTY(meta = (Singleton, Output))
 	int32 Index;
 
+	// The overall count of entries in the pose
 	UPROPERTY(meta = (Singleton, Output))
 	int32 Count;
 
@@ -823,6 +885,7 @@ struct FRigUnit_PoseLoop : public FRigUnit_HierarchyBaseMutable
 	UPROPERTY(meta = (Singleton, Output))
 	float Ratio;
 
+	// The branch to run when the loop has completed
 	UPROPERTY(meta = (Output))
 	FControlRigExecuteContext Completed;
 };
@@ -885,6 +948,7 @@ struct FRigUnit_HierarchyCreatePoseItemArray : public FRigUnit_HierarchyBase
 	UPROPERTY(meta = (Input))
 	TArray<FRigUnit_HierarchyCreatePoseItemArray_Entry> Entries;
 
+	// The resulting pose
 	UPROPERTY(meta = (Output))
 	FRigPose Pose;
 };
diff --git a/Engine/Plugins/Animation/ControlRig/Source/ControlRig/Public/Units/Execution/RigUnit_Item.h b/Engine/Plugins/Animation/ControlRig/Source/ControlRig/Public/Units/Execution/RigUnit_Item.h
--- a/Engine/Plugins/Animation/ControlRig/Source/ControlRig/Public/Units/Execution/RigUnit_Item.h
+++ b/Engine/Plugins/Animation/ControlRig/Source/ControlRig/Public/Units/Execution/RigUnit_Item.h
@@ -8,13 +8,19 @@
 
 #define UE_API CONTROLRIG_API
 
-USTRUCT(meta = (Abstract, NodeColor = "0.7 0.05 0.5", Category = "Items"))
+/*
+ * Base class for all pure item nodes
+ */
+USTRUCT(meta = (Abstract, NodeColor = "0.7 0.05 0.5", Category = "Items", DocumentationPolicy = "Strict"))
 struct FRigUnit_ItemBase : public FRigUnit
 {
 	GENERATED_BODY()
 };
 
-USTRUCT(meta = (Abstract, NodeColor = "0.7 0.05 0.5", Category = "Items"))
+/*
+ * Base class for all mutable item nodes
+ */
+USTRUCT(meta = (Abstract, NodeColor = "0.7 0.05 0.5", Category = "Items", DocumentationPolicy = "Strict"))
 struct FRigUnit_ItemBaseMutable : public FRigUnitMutable
 {
 	GENERATED_BODY()
@@ -38,9 +44,11 @@ struct FRigUnit_ItemExists : public FRigUnit_ItemBase
 	RIGVM_METHOD()
 	UE_API virtual void Execute() override;
 
+	// The item to look up
 	UPROPERTY(meta = (Input, ExpandByDefault))
 	FRigElementKey Item;
 
+	// True if the item exists
 	UPROPERTY(meta = (Output))
 	bool Exists;
 
@@ -66,15 +74,19 @@ struct FRigUnit_ItemReplace : public FRigUnit_ItemBase
 	RIGVM_METHOD()
 	UE_API virtual void Execute() override;
 
+	// The item to perform the name replace within
 	UPROPERTY(meta = (Input))
 	FRigElementKey Item;
 
+	// The old name to search for
 	UPROPERTY(meta = (Input))
 	FName Old;
 
+	// The new name to replace the old name with
 	UPROPERTY(meta = (Input))
 	FName New;
 
+	// The resulting item after the replacement
 	UPROPERTY(meta = (Output))
 	FRigElementKey Result;
 };
@@ -156,12 +168,15 @@ struct FRigUnit_ItemTypeEquals : public FRigUnit_ItemBase
 	RIGVM_METHOD()
 	UE_API virtual void Execute() override;
 
+	// The first item to compare
 	UPROPERTY(meta = (Input, ExpandByDefault))
 	FRigElementKey A;
 
+	// The second item to compare
 	UPROPERTY(meta = (Input, ExpandByDefault))
 	FRigElementKey B;
 
+	// True if the types of items A and B are equal
 	UPROPERTY(meta = (Output))
 	bool Result;
 };
@@ -183,12 +198,15 @@ struct FRigUnit_ItemTypeNotEquals : public FRigUnit_ItemBase
 	RIGVM_METHOD()
 	UE_API virtual void Execute() override;
 
+	// The first item to compare
 	UPROPERTY(meta = (Input, ExpandByDefault))
 	FRigElementKey A;
 
+	// The second item to compare
 	UPROPERTY(meta = (Input, ExpandByDefault))
 	FRigElementKey B;
 
+	// True if the types of items A and B are not equal 
 	UPROPERTY(meta = (Output))
 	bool Result;
 };
@@ -210,9 +228,11 @@ struct FRigUnit_ItemToName : public FRigUnit_ItemBase
 	RIGVM_METHOD()
 	UE_API virtual void Execute() override;
 
+	// The item to cast to a name
 	UPROPERTY(meta = (Input, ExpandByDefault))
 	FRigElementKey Value;
 
+	// The resulting name of the item
 	UPROPERTY(meta = (Output))
 	FName Result;
 };
diff --git a/Engine/Plugins/Animation/ControlRig/Source/ControlRig/Public/Units/Execution/RigUnit_RigModules.h b/Engine/Plugins/Animation/ControlRig/Source/ControlRig/Public/Units/Execution/RigUnit_RigModules.h
--- a/Engine/Plugins/Animation/ControlRig/Source/ControlRig/Public/Units/Execution/RigUnit_RigModules.h
+++ b/Engine/Plugins/Animation/ControlRig/Source/ControlRig/Public/Units/Execution/RigUnit_RigModules.h
@@ -8,13 +8,19 @@
 
 #define UE_API CONTROLRIG_API
 
-USTRUCT(meta = (Abstract, NodeColor="0.262745, 0.8, 0, 0.229412", Category = "Modules"))
+/*
+ * Base class for all pure module nodes
+ */
+USTRUCT(meta = (Abstract, NodeColor="0.262745, 0.8, 0, 0.229412", Category = "Modules", DocumentationPolicy = "Strict"))
 struct FRigUnit_RigModulesBase : public FRigUnit
 {
 	GENERATED_BODY()
 };
 
-USTRUCT(meta = (Abstract, NodeColor="0.262745, 0.8, 0, 0.229412", Category = "Modules"))
+/*
+ * Base class for all mutable module nodes
+ */
+USTRUCT(meta = (Abstract, NodeColor="0.262745, 0.8, 0, 0.229412", Category = "Modules", DocumentationPolicy = "Strict"))
 struct FRigUnit_RigModulesBaseMutable : public FRigUnitMutable
 {
 	GENERATED_BODY()
@@ -378,6 +384,7 @@ struct FRigUnit_GetItemsInModule : public FRigUnit_RigModulesBase
 	RIGVM_METHOD()
 	UE_API virtual void Execute() override;
 
+	// The types of elements to search for
 	UPROPERTY(meta = (Input))
 	ERigElementType TypeToSearch;
 
diff --git a/Engine/Plugins/Animation/ControlRig/Source/ControlRig/Public/Units/Highlevel/Harmonics/RigUnit_BoneHarmonics.h b/Engine/Plugins/Animation/ControlRig/Source/ControlRig/Public/Units/Highlevel/Harmonics/RigUnit_BoneHarmonics.h
--- a/Engine/Plugins/Animation/ControlRig/Source/ControlRig/Public/Units/Highlevel/Harmonics/RigUnit_BoneHarmonics.h
+++ b/Engine/Plugins/Animation/ControlRig/Source/ControlRig/Public/Units/Highlevel/Harmonics/RigUnit_BoneHarmonics.h
@@ -175,9 +175,11 @@ struct FRigUnit_ItemHarmonics : public FRigUnit_HighlevelBaseMutable
 	UPROPERTY(meta = (Input, Constant))
 	TArray<FRigUnit_Harmonics_TargetItem> Targets;
 
+	// The speed of the wave to use
 	UPROPERTY(meta = (Input))
 	FVector WaveSpeed;
 
+	// The frequency of the wave to use
 	UPROPERTY(meta = (Input))
 	FVector WaveFrequency;
 
@@ -185,21 +187,27 @@ struct FRigUnit_ItemHarmonics : public FRigUnit_HighlevelBaseMutable
 	UPROPERTY(meta = (Input))
 	FVector WaveAmplitude;
 
+	// The positional offset of the wave
 	UPROPERTY(meta = (Input))
 	FVector WaveOffset;
 
+	// The amount of noise to apply to the wave
 	UPROPERTY(meta = (Input))
 	FVector WaveNoise;
 
+	// The easing function to apply to the wave
 	UPROPERTY(meta = (Input))
 	ERigVMAnimEasingType WaveEase;
 
+	// The minimum value for the wave
 	UPROPERTY(meta = (Input))
 	float WaveMinimum;
 
+	// The maximum value for the wave
 	UPROPERTY(meta = (Input))
 	float WaveMaximum;
 
+	// The rotation order to use when encoding the rotation
 	UPROPERTY(meta = (Input))
 	EEulerRotationOrder RotationOrder;
 
diff --git a/Engine/Plugins/Animation/ControlRig/Source/ControlRig/Public/Units/Highlevel/Harmonics/RigUnit_ChainHarmonics.h b/Engine/Plugins/Animation/ControlRig/Source/ControlRig/Public/Units/Highlevel/Harmonics/RigUnit_ChainHarmonics.h
--- a/Engine/Plugins/Animation/ControlRig/Source/ControlRig/Public/Units/Highlevel/Harmonics/RigUnit_ChainHarmonics.h
+++ b/Engine/Plugins/Animation/ControlRig/Source/ControlRig/Public/Units/Highlevel/Harmonics/RigUnit_ChainHarmonics.h
@@ -261,27 +261,38 @@ struct FRigUnit_ChainHarmonicsPerItem : public FRigUnit_HighlevelBaseMutable
 	RIGVM_METHOD()
 	UE_API virtual void Execute() override;
 
+	// The root of the chain to apply the harmonics to
 	UPROPERTY(meta = (Input, ExpandByDefault))
 	FRigElementKey ChainRoot;
 
+	// The speed of the harmonics effects
 	UPROPERTY(meta = (Input))
 	FVector Speed;
 
+	// The reach settings
+	// Reach lets the chain "lean" towards the target trying to reach it.
 	UPROPERTY(meta = (Input))
 	FRigUnit_ChainHarmonics_Reach Reach;
 
+	// The wave settings
+	// A wave is a rocking back and forth motion to be applied to all / some axes.
 	UPROPERTY(meta = (Input))
 	FRigUnit_ChainHarmonics_Wave Wave;
 
+	// The curve to use when evaluating the wave
 	UPROPERTY(meta = (Input, Constant))
 	FRuntimeFloatCurve WaveCurve;
 
+	// The pendulum settings
+	// The harmonic pendulum uses a simple spring interpolation to allow to follow secondary motion.
 	UPROPERTY(meta = (Input))
 	FRigUnit_ChainHarmonics_Pendulum Pendulum;
 
+	// if True the debug drawing will be enabled
 	UPROPERTY(meta = (Input))
 	bool bDrawDebug;
 
+	// The world offset to be used when debug drawing
 	UPROPERTY(meta = (Input))
 	FTransform DrawWorldOffset;
 
diff --git a/Engine/Plugins/Animation/ControlRig/Source/ControlRig/Public/Units/Highlevel/Hierarchy/RigUnit_AimBone.h b/Engine/Plugins/Animation/ControlRig/Source/ControlRig/Public/Units/Highlevel/Hierarchy/RigUnit_AimBone.h
--- a/Engine/Plugins/Animation/ControlRig/Source/ControlRig/Public/Units/Highlevel/Hierarchy/RigUnit_AimBone.h
+++ b/Engine/Plugins/Animation/ControlRig/Source/ControlRig/Public/Units/Highlevel/Hierarchy/RigUnit_AimBone.h
@@ -473,13 +473,16 @@ struct FRigUnit_AimConstraintLocalSpaceOffset: public FRigUnit_HighlevelBaseMuta
 	 */
 	UPROPERTY(meta = (Input, ExpandByDefault))
 	FRigUnit_AimConstraint_WorldUp WorldUp;
-	
+
+	// The array of parents to constaint this aim to
 	UPROPERTY(meta = (Input, ExpandByDefault, DefaultArraySize = 1))
 	TArray<FConstraintParent> Parents;
 
+	// The advanced settings to use for the constraint
 	UPROPERTY(meta = (Input))
 	FRigUnit_AimConstraint_AdvancedSettings AdvancedSettings;
-	
+
+	// The weight to use when applying the resulting transform to the child
 	UPROPERTY(meta = (Input))
 	float Weight;
 	
diff --git a/Engine/Plugins/Animation/ControlRig/Source/ControlRig/Public/Units/Highlevel/Hierarchy/RigUnit_ChainInfo.h b/Engine/Plugins/Animation/ControlRig/Source/ControlRig/Public/Units/Highlevel/Hierarchy/RigUnit_ChainInfo.h
--- a/Engine/Plugins/Animation/ControlRig/Source/ControlRig/Public/Units/Highlevel/Hierarchy/RigUnit_ChainInfo.h
+++ b/Engine/Plugins/Animation/ControlRig/Source/ControlRig/Public/Units/Highlevel/Hierarchy/RigUnit_ChainInfo.h
@@ -216,7 +216,7 @@ struct FRigUnit_ChainInfo : public FRigUnit_HighlevelBase
 	float ChainLength;
 
     /* 
-	* The length of the interpolated chain
+	* The parametric length of the interpolated chain
 	*/
 	UPROPERTY(meta = (Output))
 	float ParamLength;
diff --git a/Engine/Plugins/Animation/ControlRig/Source/ControlRig/Public/Units/Highlevel/Hierarchy/RigUnit_MultiFABRIK.h b/Engine/Plugins/Animation/ControlRig/Source/ControlRig/Public/Units/Highlevel/Hierarchy/RigUnit_MultiFABRIK.h
--- a/Engine/Plugins/Animation/ControlRig/Source/ControlRig/Public/Units/Highlevel/Hierarchy/RigUnit_MultiFABRIK.h
+++ b/Engine/Plugins/Animation/ControlRig/Source/ControlRig/Public/Units/Highlevel/Hierarchy/RigUnit_MultiFABRIK.h
@@ -147,6 +147,7 @@ struct FRigUnit_MultiFABRIK : public FRigUnit_HighlevelBaseMutable
 	UPROPERTY(meta = (Input, CustomWidget = "BoneName"))
 	FName RootBone;
 
+	// The list of effectors to use for the solver
 	UPROPERTY(meta = (Input))
 	TArray<FRigUnit_MultiFABRIK_EndEffector> Effectors;
 
diff --git a/Engine/Plugins/Animation/ControlRig/Source/ControlRig/Public/Units/Highlevel/Hierarchy/RigUnit_SphericalPoseReader.h b/Engine/Plugins/Animation/ControlRig/Source/ControlRig/Public/Units/Highlevel/Hierarchy/RigUnit_SphericalPoseReader.h
--- a/Engine/Plugins/Animation/ControlRig/Source/ControlRig/Public/Units/Highlevel/Hierarchy/RigUnit_SphericalPoseReader.h
+++ b/Engine/Plugins/Animation/ControlRig/Source/ControlRig/Public/Units/Highlevel/Hierarchy/RigUnit_SphericalPoseReader.h
@@ -383,6 +383,7 @@ struct FRigUnit_SphericalPoseReader: public FRigUnit_HighlevelBaseMutable
 	UPROPERTY(meta = (Input))
 	FRigElementKey OptionalParentItem;
 
+	// The debug settings
 	UPROPERTY(meta = (Input))
 	FSphericalPoseReaderDebugSettings Debug;
 
diff --git a/Engine/Plugins/Animation/ControlRig/Source/ControlRig/Public/Units/Highlevel/Hierarchy/RigUnit_TransformConstraint.h b/Engine/Plugins/Animation/ControlRig/Source/ControlRig/Public/Units/Highlevel/Hierarchy/RigUnit_TransformConstraint.h
--- a/Engine/Plugins/Animation/ControlRig/Source/ControlRig/Public/Units/Highlevel/Hierarchy/RigUnit_TransformConstraint.h
+++ b/Engine/Plugins/Animation/ControlRig/Source/ControlRig/Public/Units/Highlevel/Hierarchy/RigUnit_TransformConstraint.h
@@ -226,21 +226,27 @@ struct FRigUnit_ParentConstraint : public FRigUnit_HighlevelBaseMutable
 	RIGVM_METHOD()
     UE_API virtual void Execute() override;
 
+	// The child to constrain
 	UPROPERTY(EditAnywhere, Category = "Constraint", meta = (Input, ExpandByDefault))
 	FRigElementKey Child;
-	
+
+	// If True the offset will be computed based on the current delta between each child and parent
 	UPROPERTY(EditAnywhere, Category = "Constraint", meta = (Input))
 	bool bMaintainOffset;
 
+	// The filter to apply on the transform
 	UPROPERTY(EditAnywhere, Category = "Constraint", meta = (Input)) 
 	FTransformFilter Filter;
-	
+
+	// The array of parents to constrain the child to
 	UPROPERTY(EditAnywhere, Category = "Constraint", meta = (Input, ExpandByDefault, DefaultArraySize = 1))
 	TArray<FConstraintParent> Parents;
-	
+
+	// The advanced settings of the constraint operation
 	UPROPERTY(EditAnywhere, Category = "Constraint", meta = (Input)) 
 	FRigUnit_ParentConstraint_AdvancedSettings AdvancedSettings;
 
+	// The weight to use when applying the resulting transform to the child
 	UPROPERTY(EditAnywhere, Category = "Constraint", meta = (Input))
 	float Weight;
 
@@ -290,13 +296,16 @@ struct FRigUnit_ParentConstraintMath : public FRigUnit_HighlevelBase
 	*/
 	UPROPERTY(EditAnywhere, Category = "Constraint", meta = (Input, ExpandByDefault))
 	FTransform Input;
-	
+
+	// The array of parents to compute the constrained transform for
 	UPROPERTY(EditAnywhere, Category = "Constraint", meta = (Input, ExpandByDefault, DefaultArraySize = 1))
 	TArray<FConstraintParent> Parents;
 
+	// The advanced settings to use for the constraint compute
 	UPROPERTY(EditAnywhere, Category = "Constraint", meta = (Input)) 
 	FRigUnit_ParentConstraintMath_AdvancedSettings AdvancedSettings;
-	
+
+	// The resulting constrained transform
 	UPROPERTY(EditAnywhere, Category = "Constraint", meta = (Output))
 	FTransform Output;
 
@@ -363,18 +372,23 @@ struct FRigUnit_PositionConstraintLocalSpaceOffset : public FRigUnit_HighlevelBa
 	RIGVM_METHOD()
 	UE_API virtual void Execute() override;
 
+	// The child to constrain
 	UPROPERTY(EditAnywhere, Category = "Constraint", meta = (Input, ExpandByDefault))
 	FRigElementKey Child;
 	
+	// If True the offset will be computed based on the current delta between each child and parent
 	UPROPERTY(EditAnywhere, Category = "Constraint", meta = (Input))
 	bool bMaintainOffset;
 
+	// The filter to apply on the transform
 	UPROPERTY(EditAnywhere, Category = "Constraint", meta = (Input)) 
 	FFilterOptionPerAxis Filter;
 
+	// The array of parents to constrain the child to
 	UPROPERTY(EditAnywhere, Category = "Constraint", meta = (Input, ExpandByDefault, DefaultArraySize = 1))
 	TArray<FConstraintParent> Parents;
 
+	// The weight to use when applying the resulting transform to the child
 	UPROPERTY(EditAnywhere, Category = "Constraint", meta = (Input))
 	float Weight;
 
@@ -471,21 +485,27 @@ struct FRigUnit_RotationConstraintLocalSpaceOffset : public FRigUnit_HighlevelBa
 	RIGVM_METHOD()
 	virtual void Execute() override;
 
+	// The child to constrain
 	UPROPERTY(EditAnywhere, Category = "Constraint", meta = (Input, ExpandByDefault))
 	FRigElementKey Child;
 	
+	// If True the offset will be computed based on the current delta between each child and parent
 	UPROPERTY(EditAnywhere, Category = "Constraint", meta = (Input))
 	bool bMaintainOffset;
 
+	// The filter to apply on the transform
 	UPROPERTY(EditAnywhere, Category = "Constraint", meta = (Input)) 
 	FFilterOptionPerAxis Filter;
 
+	// The array of parents to constrain the child to
 	UPROPERTY(EditAnywhere, Category = "Constraint", meta = (Input, ExpandByDefault, DefaultArraySize = 1))
 	TArray<FConstraintParent> Parents;
 	
+	// The advanced settings of the constraint operation
 	UPROPERTY(EditAnywhere, Category = "Constraint", meta = (Input)) 
 	FRigUnit_RotationConstraint_AdvancedSettings AdvancedSettings;
 
+	// The weight to use when applying the resulting transform to the child
 	UPROPERTY(EditAnywhere, Category = "Constraint", meta = (Input))
 	float Weight;
 
@@ -555,18 +575,23 @@ struct FRigUnit_ScaleConstraintLocalSpaceOffset : public FRigUnit_HighlevelBaseM
 	RIGVM_METHOD()
 	virtual void Execute() override;
 
+	// The child to constrain
 	UPROPERTY(EditAnywhere, Category = "Constraint", meta = (Input, ExpandByDefault))
 	FRigElementKey Child;
 	
+	// If True the offset will be computed based on the current delta between each child and parent
 	UPROPERTY(EditAnywhere, Category = "Constraint", meta = (Input))
 	bool bMaintainOffset;
 
+	// The filter to apply on the transform
 	UPROPERTY(EditAnywhere, Category = "Constraint", meta = (Input)) 
 	FFilterOptionPerAxis Filter;
 
+	// The array of parents to constrain the child to
 	UPROPERTY(EditAnywhere, Category = "Constraint", meta = (Input, ExpandByDefault, DefaultArraySize = 1))
 	TArray<FConstraintParent> Parents;
 
+	// The weight to use when applying the resulting transform to the child
 	UPROPERTY(EditAnywhere, Category = "Constraint", meta = (Input))
 	float Weight;
 
diff --git a/Engine/Plugins/Animation/ControlRig/Source/ControlRig/Public/Units/Highlevel/RigUnit_HighlevelBase.h b/Engine/Plugins/Animation/ControlRig/Source/ControlRig/Public/Units/Highlevel/RigUnit_HighlevelBase.h
--- a/Engine/Plugins/Animation/ControlRig/Source/ControlRig/Public/Units/Highlevel/RigUnit_HighlevelBase.h
+++ b/Engine/Plugins/Animation/ControlRig/Source/ControlRig/Public/Units/Highlevel/RigUnit_HighlevelBase.h
@@ -12,13 +12,19 @@ enum class EControlRigVectorKind : uint8
 	Location
 };
 
-USTRUCT(meta=(Abstract, Category="Debug", NodeColor = "0.462745 1.0 0.329412"))
+/*
+ * Base class for all pure high level nodes
+ */
+USTRUCT(meta=(Abstract, Category="Debug", NodeColor = "0.462745 1.0 0.329412", DocumentationPolicy = "Strict"))
 struct FRigUnit_HighlevelBase : public FRigUnit
 {
 	GENERATED_BODY()
 };
 
-USTRUCT(meta=(Abstract, Category="Debug", NodeColor = "0 0.364706 1.0"))
+/*
+ * Base class for all mutable high level nodes
+ */
+USTRUCT(meta=(Abstract, Category="Debug", NodeColor = "0 0.364706 1.0", DocumentationPolicy = "Strict"))
 struct FRigUnit_HighlevelBaseMutable : public FRigUnitMutable
 {
 	GENERATED_BODY()
diff --git a/Engine/Plugins/Animation/ControlRigSpline/Source/ControlRigSpline/Public/ControlRigSplineUnits.h b/Engine/Plugins/Animation/ControlRigSpline/Source/ControlRigSpline/Public/ControlRigSplineUnits.h
--- a/Engine/Plugins/Animation/ControlRigSpline/Source/ControlRigSpline/Public/ControlRigSplineUnits.h
+++ b/Engine/Plugins/Animation/ControlRigSpline/Source/ControlRigSpline/Public/ControlRigSplineUnits.h
@@ -9,16 +9,28 @@
 
 #define UE_API CONTROLRIGSPLINE_API
 
-USTRUCT(meta = (Abstract, NodeColor = "0.3 0.1 0.1"))
+/*
+ * The base class for all control rig spline nodes
+ */
+USTRUCT(meta = (Abstract, Category = "Splines", Varying, NodeColor = "0.737911 0.099899 0.099899", DocumentationPolicy = "Strict"))
 struct FRigUnit_ControlRigSplineBase : public FRigUnit
 {
 	GENERATED_BODY()
 };
 
+/*
+ * The base class for all mutable control rig spline nodes
+ */
+USTRUCT(meta = (Abstract, Category = "Splines", Varying, NodeColor = "0.737911 0.099899 0.099899", DocumentationPolicy = "Strict"))
+struct FRigUnit_ControlRigSplineBaseMutable : public FRigUnitMutable
+{
+	GENERATED_BODY()
+};
+
 /*
  * Creates a Spline curve from an array of positions
  */
-USTRUCT(meta = (DisplayName = "Spline From Points", Keywords="Spline From Positions", Category = "Splines", Varying, NodeColor="0.737911 0.099899 0.099899"))
+USTRUCT(meta = (DisplayName = "Spline From Points", Keywords="Spline From Positions"))
 struct FRigUnit_ControlRigSplineFromPoints : public FRigUnit_ControlRigSplineBase
 {
 	GENERATED_BODY()
@@ -36,32 +48,39 @@ struct FRigUnit_ControlRigSplineFromPoints : public FRigUnit_ControlRigSplineBas
 	RIGVM_METHOD()
 	UE_API virtual void Execute() override;
 
+	// The input points to form the spline
 	UPROPERTY(meta = (Input))
 	TArray<FVector> Points;
 
+	// The mode to use for the spline
 	UPROPERTY(meta = (Input))
 	ESplineType SplineMode;
 
+	// If True the spline will be closed 
 	UPROPERTY(meta = (Input))
 	bool bClosed;
 
+	// Specifies the detail per segment of the spline
 	UPROPERTY(meta = (Input, ClampMin = "1"))
 	int32 SamplesPerSegment;
 
+	// The amount of compression to apply
 	UPROPERTY(meta = (Input))
 	float Compression;
 
+	// The amount of stretch to allow for the spline
 	UPROPERTY(meta = (Input))
 	float Stretch;
 
+	// The resulting spline
 	UPROPERTY(meta = (Output))
 	FControlRigSpline Spline;
 };
 
 /*
- * Creates a Spline curve from an array of positions
+ * Creates a Spline curve from an array of transforms
  */
-USTRUCT(meta = (DisplayName = "Spline From Transforms", Keywords="Spline From Transforms", Category = "Splines", Varying, NodeColor="0.737911 0.099899 0.099899"))
+USTRUCT(meta = (DisplayName = "Spline From Transforms", Keywords="Spline From Transforms"))
 struct FRigUnit_ControlRigSplineFromTransforms : public FRigUnit_ControlRigSplineBase
 {
 	GENERATED_BODY()
@@ -79,24 +98,31 @@ struct FRigUnit_ControlRigSplineFromTransforms : public FRigUnit_ControlRigSplin
 	RIGVM_METHOD()
 	UE_API virtual void Execute() override;
 
+	// The input transforms to build the spline from
 	UPROPERTY(meta = (Input))
 	TArray<FTransform> Transforms;
 
+	// The mode to use for the spline
 	UPROPERTY(meta = (Input))
 	ESplineType SplineMode;
 
+	// If True the spline will be closed 
 	UPROPERTY(meta = (Input))
 	bool bClosed;
 
+	// Specifies the detail per segment of the spline
 	UPROPERTY(meta = (Input))
 	int32 SamplesPerSegment;
 
+	// The amount of compression to apply
 	UPROPERTY(meta = (Input))
 	float Compression;
 
+	// The amount of stretch to allow for the spline
 	UPROPERTY(meta = (Input))
 	float Stretch;
 
+	// The resulting spline
 	UPROPERTY(meta = (Output))
 	FControlRigSpline Spline;
 };
@@ -104,8 +130,8 @@ struct FRigUnit_ControlRigSplineFromTransforms : public FRigUnit_ControlRigSplin
 /*
  * Set the points of a spline, given a spline and an array of positions
  */
-USTRUCT(meta = (DisplayName = "Set Spline Points", Category = "Splines", Varying, NodeColor = "0.0 0.36470600962638855 1.0"))
-struct FRigUnit_SetSplinePoints : public FRigUnitMutable
+USTRUCT(meta = (DisplayName = "Set Spline Points"))
+struct FRigUnit_SetSplinePoints : public FRigUnit_ControlRigSplineBaseMutable
 {
 	GENERATED_BODY()
 
@@ -118,18 +144,21 @@ struct FRigUnit_SetSplinePoints : public FRigUnitMutable
 	RIGVM_METHOD()
 	UE_API virtual void Execute() override;
 
+	// The points to set on the spline. The assumption is that the
+	// number of points provided matches the points on the spline.
 	UPROPERTY(meta = (Input))
 	TArray<FVector> Points;
 
+	// The spline to be updated
 	UPROPERTY(meta = (Input, Output))
 	FControlRigSpline Spline;
 };
 
 /*
- * Set the points of a spline, given a spline and an array of positions
+ * Set the points of a spline, given a spline and an array of transforms
  */
-USTRUCT(meta = (DisplayName = "Set Spline Transforms", Category = "Splines", Varying, NodeColor = "0.0 0.36470600962638855 1.0"))
-struct FRigUnit_SetSplineTransforms : public FRigUnitMutable
+USTRUCT(meta = (DisplayName = "Set Spline Transforms"))
+struct FRigUnit_SetSplineTransforms : public FRigUnit_ControlRigSplineBaseMutable
 {
 	GENERATED_BODY()
 
@@ -142,9 +171,12 @@ struct FRigUnit_SetSplineTransforms : public FRigUnitMutable
 	RIGVM_METHOD()
 	UE_API virtual void Execute() override;
 
+	// The transforms interpreted as positions to set on the spline. The assumption is that the
+	// number of transforms provided matches the points on the spline.
 	UPROPERTY(meta = (Input))
 	TArray<FTransform> Transforms;
 
+	// The spline to be updated
 	UPROPERTY(meta = (Input, Output))
 	FControlRigSpline Spline;
 };
@@ -152,7 +184,7 @@ struct FRigUnit_SetSplineTransforms : public FRigUnitMutable
 /*
  * Retrieves the position from a given Spline and U value
  */
-USTRUCT(meta = (DisplayName = "Position From Spline", Keywords="Point From Spline", Category = "Splines", NodeColor="0.737911 0.099899 0.099899"))
+USTRUCT(meta = (DisplayName = "Position From Spline", Keywords="Point From Spline"))
 struct FRigUnit_PositionFromControlRigSpline : public FRigUnit_ControlRigSplineBase
 {
 	GENERATED_BODY()
@@ -167,12 +199,15 @@ struct FRigUnit_PositionFromControlRigSpline : public FRigUnit_ControlRigSplineB
 	RIGVM_METHOD()
 	UE_API virtual void Execute() override;
 
+	// The spline to evaluate
 	UPROPERTY(meta = (Input))
 	FControlRigSpline Spline;
 
+	// The U value along the spline to evaluate
 	UPROPERTY(meta = (Input, ClampMin = "0.0", ClampMax = "1.0"))
 	float U;
 
+	// The resulting position on the spline
 	UPROPERTY(meta = (Output))
 	FVector Position;
 };
@@ -180,7 +215,7 @@ struct FRigUnit_PositionFromControlRigSpline : public FRigUnit_ControlRigSplineB
 /*
  * Retrieves the transform from a given Spline and U value based on the given Up Vector and Roll
  */
-USTRUCT(meta = (DisplayName = "Transform From Spline (with UpVector)", Category = "Splines", NodeColor="0.737911 0.099899 0.099899"))
+USTRUCT(meta = (DisplayName = "Transform From Spline (with UpVector)"))
 struct FRigUnit_TransformFromControlRigSpline : public FRigUnit_ControlRigSplineBase
 {
 	GENERATED_BODY()
@@ -197,26 +232,31 @@ struct FRigUnit_TransformFromControlRigSpline : public FRigUnit_ControlRigSpline
 	RIGVM_METHOD()
 	UE_API virtual void Execute() override;
 
+	// The spline to evaluate
 	UPROPERTY(meta = (Input))
 	FControlRigSpline Spline;
 
+	// The up-vector to use to build the transform's rotation
 	UPROPERTY(meta = (Input))
 	FVector UpVector;
 
+	// The roll to apply to the resulting rotation
 	UPROPERTY(meta = (Input))
 	float Roll;
 
+	// The U value along the spline to evaluate
 	UPROPERTY(meta = (Input))
 	float U;
 
+	// The resulting composed transform
 	UPROPERTY(meta = (Output))
 	FTransform Transform;
 };
 
 /*
- * Retrieves the transform from a given Spline and U value based on the given Up Vector and Roll
+ * Retrieves the transform from a given Spline and U value based on the given primary and secondary axis
  */
-USTRUCT(meta = (DisplayName = "Transform From Spline", Category = "Splines", NodeColor="0.737911 0.099899 0.099899"))
+USTRUCT(meta = (DisplayName = "Transform From Spline"))
 struct FRigUnit_TransformFromControlRigSpline2 : public FRigUnit_ControlRigSplineBase
 {
 	GENERATED_BODY()
@@ -233,18 +273,23 @@ struct FRigUnit_TransformFromControlRigSpline2 : public FRigUnit_ControlRigSplin
 	RIGVM_METHOD()
 	UE_API virtual void Execute() override;
 
+	// The spline to evaluate
 	UPROPERTY(meta = (Input))
 	FControlRigSpline Spline;
 
+	// The U value along the spline to evaluate
 	UPROPERTY(meta = (Input))
 	float U;
 
+	// The primary axis to use when building the transform
 	UPROPERTY(meta = (Input))
 	FVector PrimaryAxis;
 
+	// The secondary axis to use when building the transform 
 	UPROPERTY(meta = (Input))
 	FVector SecondaryAxis;
 
+	// The resulting composed transform
 	UPROPERTY(meta = (Output))
 	FTransform Transform;
 };
@@ -252,7 +297,7 @@ struct FRigUnit_TransformFromControlRigSpline2 : public FRigUnit_ControlRigSplin
 /*
  * Retrieves the tangent from a given Spline and U value
  */
-USTRUCT(meta = (DisplayName = "Tangent From Spline", Category = "Splines", NodeColor="0.737911 0.099899 0.099899"))
+USTRUCT(meta = (DisplayName = "Tangent From Spline"))
 struct FRigUnit_TangentFromControlRigSpline : public FRigUnit_ControlRigSplineBase
 {
 	GENERATED_BODY()
@@ -267,12 +312,15 @@ struct FRigUnit_TangentFromControlRigSpline : public FRigUnit_ControlRigSplineBa
 	RIGVM_METHOD()
 	UE_API virtual void Execute() override;
 
+	// The spline to evaluate
 	UPROPERTY(meta = (Input))
 	FControlRigSpline Spline;
 
+	// The U value along the spline to evaluate
 	UPROPERTY(meta = (Input, ClampMin = "0.0", ClampMax = "1.0"))
 	float U;
 
+	// The tangent at the evaluated location on the spline
 	UPROPERTY(meta = (Output))
 	FVector Tangent;
 };
@@ -280,8 +328,8 @@ struct FRigUnit_TangentFromControlRigSpline : public FRigUnit_ControlRigSplineBa
 /*
  * Draws the given spline in the viewport
  */
-USTRUCT(meta = (DisplayName = "Draw Spline", Category = "Splines", NodeColor = "0.83077 0.846873 0.049707"))
-struct FRigUnit_DrawControlRigSpline : public FRigUnitMutable
+USTRUCT(meta = (DisplayName = "Draw Spline"))
+struct FRigUnit_DrawControlRigSpline : public FRigUnit_ControlRigSplineBaseMutable
 {
 	GENERATED_BODY()
 
@@ -296,24 +344,28 @@ struct FRigUnit_DrawControlRigSpline : public FRigUnitMutable
 	RIGVM_METHOD()
 	UE_API virtual void Execute() override;
 
+	// The spline to draw
 	UPROPERTY(meta = (Input))
 	FControlRigSpline Spline;
 
+	// The color to use for the debug draw
 	UPROPERTY(meta = (Input))
 	FLinearColor Color;
 
+	// The line thickness to use when drawing the spline
 	UPROPERTY(meta = (Input))
 	float Thickness;
 
+	// The detail to use to render the spline
 	UPROPERTY(meta = (Input))
 	int32 Detail;
 };
 
 /*
- * Retrieves the length from a given Splin
+ * Retrieves the length from a given Spline
  */
-USTRUCT(meta = (DisplayName = "Get Length Of Spline", Category = "Splines", NodeColor="0.737911 0.099899 0.099899"))
-struct FRigUnit_GetLengthControlRigSpline : public FRigUnit
+USTRUCT(meta = (DisplayName = "Get Length Of Spline"))
+struct FRigUnit_GetLengthControlRigSpline : public FRigUnit_ControlRigSplineBase
 {
 	GENERATED_BODY()
 
@@ -326,18 +378,20 @@ struct FRigUnit_GetLengthControlRigSpline : public FRigUnit
 	RIGVM_METHOD()
 	UE_API virtual void Execute() override;
 
+	// The spline to evaluate
 	UPROPERTY(meta = (Input))
 	FControlRigSpline Spline;
 
+	// The resulting length of the spline
 	UPROPERTY(meta = (Output))
 	float Length;
 };
 
 /*
- * Retrieves the length from a given Splin
+ * Retrieves the length from a given Spline
  */
-USTRUCT(meta = (DisplayName = "Get Length At Param Of Spline", Category = "Splines", NodeColor="0.737911 0.099899 0.099899"))
-struct FRigUnit_GetLengthAtParamControlRigSpline : public FRigUnit
+USTRUCT(meta = (DisplayName = "Get Length At Param Of Spline"))
+struct FRigUnit_GetLengthAtParamControlRigSpline : public FRigUnit_ControlRigSplineBase
 {
 	GENERATED_BODY()
 
@@ -351,12 +405,15 @@ struct FRigUnit_GetLengthAtParamControlRigSpline : public FRigUnit
 	RIGVM_METHOD()
 	UE_API virtual void Execute() override;
 
+	// The spline to evaluate
 	UPROPERTY(meta = (Input))
 	FControlRigSpline Spline;
 
+	// The U value along the spline to evaluate
 	UPROPERTY(meta = (Input, ClampMin = "0.0", ClampMax = "1.0"))
 	float U;
 
+	// The resulting length 
 	UPROPERTY(meta = (Output))
 	float Length;
 };
@@ -366,7 +423,7 @@ struct FRigUnit_GetLengthAtParamControlRigSpline : public FRigUnit
  * Additionally provides rotational control matching the features of the Distribute Rotation node.
  */
 USTRUCT(meta=(DisplayName="Fit Chain on Spline Curve", Category="Hierarchy", Keywords="Fit,Resample,Spline", Deprecated = "5.0"))
-struct FRigUnit_FitChainToSplineCurve : public FRigUnit_HighlevelBaseMutable
+struct FRigUnit_FitChainToSplineCurve : public FRigUnit_ControlRigSplineBaseMutable
 {
 	GENERATED_BODY()
 
@@ -438,7 +495,7 @@ struct FRigUnit_FitChainToSplineCurve : public FRigUnit_HighlevelBaseMutable
 	FVector SecondaryAxis;
 
 	/**
-	 * The the position of the pole vector used for aligning the secondary axis.
+	 * The position of the pole vector used for aligning the secondary axis.
 	 * Only has an effect if the secondary axis is set.
 	 */
 	UPROPERTY(meta = (Input))
@@ -470,6 +527,7 @@ struct FRigUnit_FitChainToSplineCurve : public FRigUnit_HighlevelBaseMutable
 	UPROPERTY(meta = (Input, Constant))
 	bool bPropagateToChildren;
 
+	// The debug settings to use
 	UPROPERTY(meta = (Input, DetailsOnly))
 	FRigUnit_FitChainToCurve_DebugSettings DebugSettings;
 
@@ -484,8 +542,8 @@ struct FRigUnit_FitChainToSplineCurve : public FRigUnit_HighlevelBaseMutable
  * Fits a given chain to a spline curve.
  * Additionally provides rotational control matching the features of the Distribute Rotation node.
  */
-USTRUCT(meta=(DisplayName="Fit Chain on Spline Curve", Category="Splines", Keywords="Fit,Resample,Spline"))
-struct FRigUnit_FitChainToSplineCurveItemArray : public FRigUnit_HighlevelBaseMutable
+USTRUCT(meta=(DisplayName="Fit Chain on Spline Curve", Keywords="Fit,Resample,Spline"))
+struct FRigUnit_FitChainToSplineCurveItemArray : public FRigUnit_ControlRigSplineBaseMutable
 {
 	GENERATED_BODY()
 
@@ -557,7 +615,7 @@ struct FRigUnit_FitChainToSplineCurveItemArray : public FRigUnit_HighlevelBaseMu
 	FVector SecondaryAxis;
 
 	/**
-	 * The the position of the pole vector used for aligning the secondary axis.
+	 * The position of the pole vector used for aligning the secondary axis.
 	 * Only has an effect if the secondary axis is set.
 	 */
 	UPROPERTY(meta = (Input))
@@ -589,6 +647,7 @@ struct FRigUnit_FitChainToSplineCurveItemArray : public FRigUnit_HighlevelBaseMu
 	UPROPERTY(meta = (Input, Constant))
 	bool bPropagateToChildren;
 
+	// The debug settings to use
 	UPROPERTY(meta = (Input, DetailsOnly))
 	FRigUnit_FitChainToCurve_DebugSettings DebugSettings;
 
@@ -624,8 +683,8 @@ struct FRigUnit_SplineConstraint_WorkData
  * Fits a given chain to a spline curve.
  * Additionally provides rotational control matching the features of the Distribute Rotation node.
  */
-USTRUCT(meta=(DisplayName="Spline Constraint", Category="Splines", Keywords="Fit,Resample,Spline"))
-struct FRigUnit_SplineConstraint : public FRigUnit_HighlevelBaseMutable
+USTRUCT(meta=(DisplayName="Spline Constraint", Keywords="Fit,Resample,Spline"))
+struct FRigUnit_SplineConstraint : public FRigUnit_ControlRigSplineBaseMutable
 {
 	GENERATED_BODY()
 
@@ -672,9 +731,11 @@ struct FRigUnit_SplineConstraint : public FRigUnit_HighlevelBaseMutable
 	UPROPERTY(meta = (Input, ClampMin = "0.0", ClampMax = "1.0"))
 	float Maximum;
 
+	// The primary axis to use when building transforms
 	UPROPERTY(meta = (Input))
 	FVector PrimaryAxis;
 
+	// The secondary axis to use when building transforms
 	UPROPERTY(meta = (Input))
 	FVector SecondaryAxis;
 
@@ -695,7 +756,7 @@ struct FRigUnit_SplineConstraint : public FRigUnit_HighlevelBaseMutable
  * Fits a given spline curve to a chain.
  */
 USTRUCT(meta=(DisplayName="Fit Spline Curve on Chain", Category="Hierarchy", Keywords="Fit,Resample,Spline", Deprecated = "5.0"))
-struct FRigUnit_FitSplineCurveToChain : public FRigUnit_HighlevelBaseMutable
+struct FRigUnit_FitSplineCurveToChain : public FRigUnit_ControlRigSplineBaseMutable
 {
 	GENERATED_BODY()
 
@@ -726,8 +787,8 @@ struct FRigUnit_FitSplineCurveToChain : public FRigUnit_HighlevelBaseMutable
 /**
  * Fits a given spline curve to a chain.
  */
-USTRUCT(meta=(DisplayName="Fit Spline Curve on Chain", Category="Splines", Keywords="Fit,Resample,Spline"))
-struct FRigUnit_FitSplineCurveToChainItemArray : public FRigUnit_HighlevelBaseMutable
+USTRUCT(meta=(DisplayName="Fit Spline Curve on Chain", Keywords="Fit,Resample,Spline"))
+struct FRigUnit_FitSplineCurveToChainItemArray : public FRigUnit_ControlRigSplineBaseMutable
 {
 	GENERATED_BODY()
 
@@ -755,7 +816,7 @@ struct FRigUnit_FitSplineCurveToChainItemArray : public FRigUnit_HighlevelBaseMu
 /*
  * Retrieves the closest U value from a given Spline and a position
  */
-USTRUCT(meta = (DisplayName = "Closest Parameter From Spline", Category = "Splines", NodeColor="0.737911 0.099899 0.099899"))
+USTRUCT(meta = (DisplayName = "Closest Parameter From Spline"))
 struct FRigUnit_ClosestParameterFromControlRigSpline : public FRigUnit_ControlRigSplineBase
 {
 	GENERATED_BODY()
@@ -770,12 +831,15 @@ struct FRigUnit_ClosestParameterFromControlRigSpline : public FRigUnit_ControlRi
 	RIGVM_METHOD()
 	UE_API virtual void Execute() override;
 
+	// The spline to evaluate
 	UPROPERTY(meta = (Input))
 	FControlRigSpline Spline;
 
+	// The position to evaluate
 	UPROPERTY(meta = (Input))
 	FVector Position;
 
+	// The U value at the closest location on the spline
 	UPROPERTY(meta = (Output))
 	float U;
 	
@@ -784,7 +848,7 @@ struct FRigUnit_ClosestParameterFromControlRigSpline : public FRigUnit_ControlRi
 /*
  * Returns the U parameter of a spline given a length percentage (0.0 - 1.0)
  */
-USTRUCT(meta = (DisplayName = "Parameter At Length Percentage", Category = "Splines"))
+USTRUCT(meta = (DisplayName = "Parameter At Length Percentage"))
 struct FRigUnit_ParameterAtPercentage : public FRigUnit_ControlRigSplineBase
 {
 	GENERATED_BODY()
@@ -799,12 +863,15 @@ struct FRigUnit_ParameterAtPercentage : public FRigUnit_ControlRigSplineBase
 	RIGVM_METHOD()
 	UE_API virtual void Execute() override;
 
+	// The spline to evaluate
 	UPROPERTY(meta = (Input))
 	FControlRigSpline Spline;
 
+	// The percentage (0.0 - 1.0) to evaluate
 	UPROPERTY(meta = (Input))
 	float Percentage;
 
+	// The U value on the spline for the given percentage
 	UPROPERTY(meta = (Output))
 	float U;
 	
diff --git a/Engine/Plugins/Experimental/Animation/Locomotor/Source/Locomotor/Public/RigUnit_Locomotor.h b/Engine/Plugins/Experimental/Animation/Locomotor/Source/Locomotor/Public/RigUnit_Locomotor.h
--- a/Engine/Plugins/Experimental/Animation/Locomotor/Source/Locomotor/Public/RigUnit_Locomotor.h
+++ b/Engine/Plugins/Experimental/Animation/Locomotor/Source/Locomotor/Public/RigUnit_Locomotor.h
@@ -340,6 +340,9 @@ struct LOCOMOTOR_API FFootSet
 	float PhaseOffset = 0.f;
 };
 
+/*
+ * A node providing access to the locomotor functionality for foot step generation.
+ */
 USTRUCT(meta=(DisplayName="Locomotor", Category="Simulation", Keywords="Walk Step Walker Walking Locomotion Procedural"))
 struct LOCOMOTOR_API FRigUnit_Locomotor: public FRigUnit_HighlevelBaseMutable
 {
diff --git a/Engine/Plugins/Experimental/LiveLinkControlRig/Source/LiveLinkControlRig/Public/LiveLinkRigUnits.h b/Engine/Plugins/Experimental/LiveLinkControlRig/Source/LiveLinkControlRig/Public/LiveLinkRigUnits.h
--- a/Engine/Plugins/Experimental/LiveLinkControlRig/Source/LiveLinkControlRig/Public/LiveLinkRigUnits.h
+++ b/Engine/Plugins/Experimental/LiveLinkControlRig/Source/LiveLinkControlRig/Public/LiveLinkRigUnits.h
@@ -14,7 +14,10 @@ namespace LiveLinkControlRigUtilities
 	ILiveLinkClient* TryGetLiveLinkClient();
 }
 
-USTRUCT(meta = (Abstract, NodeColor = "0.3 0.1 0.1"))
+/*
+ * The base class for all live link control rig nodes
+ */
+USTRUCT(meta = (Abstract, NodeColor = "0.3 0.1 0.1", DocumentationPolicy = "Strict"))
 struct FRigUnit_LiveLinkBase : public FRigUnit
 {
 	GENERATED_BODY()
@@ -32,18 +35,23 @@ struct FRigUnit_LiveLinkEvaluteFrameAnimation : public FRigUnit_LiveLinkBase
 	RIGVM_METHOD()
 	UE_API virtual void Execute() override;
 
+	// The name of the subject to evaluate a frame for
 	UPROPERTY(meta = (Input))
 	FName SubjectName;
 
+	// If True debug data will be drawn for the subject
 	UPROPERTY(meta = (Input))
 	bool bDrawDebug = false;
 
+	// The color to use for the debug drawing
 	UPROPERTY(meta = (Input))
 	FLinearColor DebugColor = FLinearColor::Red;
 
+	// The world offset to use when drawing the debug data
 	UPROPERTY(meta = (Input))
 	FTransform DebugDrawOffset;
 
+	// The resulting subject's frame
 	UPROPERTY(meta = (Output))
 	FSubjectFrameHandle SubjectFrame;
 };
@@ -60,15 +68,19 @@ struct FRigUnit_LiveLinkGetTransformByName : public FRigUnit_LiveLinkBase
 	RIGVM_METHOD()
 	UE_API virtual void Execute() override;
 
+	// The frame to receive a transform from
 	UPROPERTY(meta = (Input))
 	FSubjectFrameHandle SubjectFrame;
 
+	// The name of the transform to retrieve
 	UPROPERTY(meta = (Input))
 	FName TransformName;
 
+	// The space to retrieve the transform in
 	UPROPERTY(meta = (Input))
 	ERigVMTransformSpace Space = ERigVMTransformSpace::LocalSpace;
 
+	// The resulting transform
 	UPROPERTY(meta = (Output))
 	FTransform Transform;
 };
@@ -85,12 +97,15 @@ struct FRigUnit_LiveLinkGetParameterValueByName : public FRigUnit_LiveLinkBase
 	RIGVM_METHOD()
 	UE_API virtual void Execute() override;
 
+	// The frame to retrieve the parameter from
 	UPROPERTY(meta = (Input))
 	FSubjectFrameHandle SubjectFrame;
 
+	// The name of the parameter to retrieve
 	UPROPERTY(meta = (Input))
 	FName ParameterName;
 
+	// The resulting value of the parameter
 	UPROPERTY(meta = (Output))
 	float Value = 0.f;
 };
@@ -107,18 +122,23 @@ struct FRigUnit_LiveLinkEvaluteFrameTransform : public FRigUnit_LiveLinkBase
 	RIGVM_METHOD()
 	UE_API virtual void Execute() override;
 
+	// The name of the subject to evaluate
 	UPROPERTY(meta = (Input))
 	FName SubjectName;
 
+	// If True debug data will be drawn
 	UPROPERTY(meta = (Input))
 	bool bDrawDebug = false;
 
+	// The color to use for the debug drawing
 	UPROPERTY(meta = (Input))
 	FLinearColor DebugColor = FLinearColor::Red;
 
+	// The world offset to use for the debug drawing
 	UPROPERTY(meta = (Input))
 	FTransform DebugDrawOffset;
 
+	// The resulting transform of the subject
 	UPROPERTY(meta = (Output))
 	FTransform Transform;
 };
@@ -134,12 +154,15 @@ struct FRigUnit_LiveLinkEvaluateBasicValue : public FRigUnit_LiveLinkBase
 	RIGVM_METHOD()
 	UE_API virtual void Execute() override;
 
+	// The name of the subject to evaluate
 	UPROPERTY(meta = (Input))
 	FName SubjectName;
 
+	// The name of the property to evaluate
 	UPROPERTY(meta = (Input))
 	FName PropertyName;
 
+	// The resulting value of the evaluated property
 	UPROPERTY(meta = (Output))
 	float Value = 0.f;
 };
diff --git a/Engine/Plugins/Experimental/MoverAnimNext/Source/MoverAnimNext/Private/Graph/RigUnit_GenerateMoverTrajectory.h b/Engine/Plugins/Experimental/MoverAnimNext/Source/MoverAnimNext/Private/Graph/RigUnit_GenerateMoverTrajectory.h
--- a/Engine/Plugins/Experimental/MoverAnimNext/Source/MoverAnimNext/Private/Graph/RigUnit_GenerateMoverTrajectory.h
+++ b/Engine/Plugins/Experimental/MoverAnimNext/Source/MoverAnimNext/Private/Graph/RigUnit_GenerateMoverTrajectory.h
@@ -10,6 +10,9 @@
 
 class UMoverComponent;
 
+/*
+ * A utility node to generate a trajectory from a mover component
+ */
 USTRUCT(meta = (DisplayName = "Generate Trajectory from Mover", Category = "Mover", NodeColor = "0, 1, 1", Keywords = "Mover, Trajectory"))
 struct FRigUnit_GenerateMoverTrajectory : public FRigUnit_AnimNextBase
 {
@@ -18,25 +21,31 @@ struct FRigUnit_GenerateMoverTrajectory : public FRigUnit_AnimNextBase
 	RIGVM_METHOD()
 	void Execute();
 
+	// The mover component to generate the trajectory for
 	UPROPERTY(EditAnywhere, Category = "Trajectory", meta = (Input))
 	TObjectPtr<UMoverComponent> MoverComponent;
 
-	UPROPERTY(EditAnywhere, Category = "Trajectory", meta = (Input))
 	// This should be the most recent simulation time that was used to get us to our current state
+	UPROPERTY(EditAnywhere, Category = "Trajectory", meta = (Input))
 	float DeltaTime = 0.f;
 
+	// The sampling interval to use for the history
 	UPROPERTY(EditAnywhere, Category = "Trajectory", meta = (Input))
 	float HistorySamplingInterval = -1.f;
 
+	// The upper limit of samples to use for the history
 	UPROPERTY(EditAnywhere, Category = "Trajectory", meta = (Input))
 	int32 NumHistorySamples = 30;
 
+	// The interval to use for prediction
 	UPROPERTY(EditAnywhere, Category = "Trajectory", meta = (Input))
 	float PredictionSamplingInterval = 0.1f;
 
+	// The upper limit of samples to use for prediction
 	UPROPERTY(EditAnywhere, Category = "Trajectory", meta = (Input))
 	int32 NumPredictionSamples = 15;
 
+	// The resulting trajectory 
 	UPROPERTY(EditAnywhere, Category = "Trajectory", meta = (Input, Output))
 	FTransformTrajectory InOutTrajectory;
 	
diff --git a/Engine/Plugins/Experimental/MoverAnimNext/Source/MoverAnimNext/Private/Graph/RigUnit_TrajectoryAnalysis.h b/Engine/Plugins/Experimental/MoverAnimNext/Source/MoverAnimNext/Private/Graph/RigUnit_TrajectoryAnalysis.h
--- a/Engine/Plugins/Experimental/MoverAnimNext/Source/MoverAnimNext/Private/Graph/RigUnit_TrajectoryAnalysis.h
+++ b/Engine/Plugins/Experimental/MoverAnimNext/Source/MoverAnimNext/Private/Graph/RigUnit_TrajectoryAnalysis.h
@@ -8,6 +8,9 @@
 
 #include "RigUnit_TrajectoryAnalysis.generated.h"
 
+/*
+ * Returns a sample of a trajectory at a given time
+ */
 USTRUCT(meta = (DisplayName = "Get Trajectory Sample At Time", Category = "Trajectory", NodeColor = "0, 1, 1", Keywords = "Trajectory"))
 struct FRigUnit_GetTrajectorySampleAtTime : public FRigUnit_AnimNextBase
 {
@@ -16,19 +19,26 @@ struct FRigUnit_GetTrajectorySampleAtTime : public FRigUnit_AnimNextBase
 	RIGVM_METHOD()
 	void Execute();
 
+	// The resulting trajectory sample
 	UPROPERTY(EditAnywhere, Category = "Trajectory", meta = (Output))
 	FTransformTrajectorySample OutTrajectorySample;
 
+	// The input trajectory to analyze
 	UPROPERTY(EditAnywhere, Category = "Trajectory", meta = (Input))
 	FTransformTrajectory InTrajectory;
 
+	// The time at which to analyze the trajectory
 	UPROPERTY(EditAnywhere, Category = "Trajectory", meta = (Input))
 	float Time = 0.0f;
-
+	
+	// If True the trajectory will be extrapolated beyond its bounds
 	UPROPERTY(EditAnywhere, Category = "Trajectory", meta = (Input))
 	bool bExtrapolate = false;
 };
 
+/*
+ * Computes the velocity between two times on a trajectory
+ */
 USTRUCT(meta = (DisplayName = "Get Trajectory Velocity", Category = "Trajectory", NodeColor = "0, 1, 1", Keywords = "Trajectory"))
 struct FRigUnit_GetTrajectoryVelocity : public FRigUnit_AnimNextBase
 {
@@ -37,22 +47,30 @@ struct FRigUnit_GetTrajectoryVelocity : public FRigUnit_AnimNextBase
 	RIGVM_METHOD()
 	void Execute();
 
+	// The input trajectory to analyze
 	UPROPERTY(EditAnywhere, Category = "Trajectory", meta = (Input))
 	FTransformTrajectory InTrajectory;
 
+	// The resulting velocity between the two provided times
 	UPROPERTY(EditAnywhere, Category = "Trajectory", meta = (Output))
 	FVector OutVelocity = FVector::ZeroVector;
 
+	// The first time to analyze
 	UPROPERTY(EditAnywhere, Category = "Trajectory", meta = (Input))
 	float Time1 = 0.0f;
 
+	// The second time to analyze
 	UPROPERTY(EditAnywhere, Category = "Trajectory", meta = (Input))
 	float Time2 = 0.0f;
 
+	// If True the trajectory will be extrapolated beyond its bounds
 	UPROPERTY(EditAnywhere, Category = "Trajectory", meta = (Input))
 	bool bExtrapolate = false;
 };
 
+/*
+ * Computes the angular velocity between two times on a trajectory
+ */
 USTRUCT(meta = (DisplayName = "Get Trajectory Angular Velocity", Category = "Trajectory", NodeColor = "0, 1, 1", Keywords = "Trajectory"))
 struct FRigUnit_GetTrajectoryAngularVelocity : public FRigUnit_AnimNextBase
 {
@@ -61,18 +79,23 @@ struct FRigUnit_GetTrajectoryAngularVelocity : public FRigUnit_AnimNextBase
 	RIGVM_METHOD()
 	void Execute();
 
+	// The input trajectory to analyze
 	UPROPERTY(EditAnywhere, Category = "Trajectory", meta = (Input))
 	FTransformTrajectory InTrajectory;
 
+	// The resulting angular velocity between the two provided times
 	UPROPERTY(EditAnywhere, Category = "Trajectory", meta = (Output))
 	FVector OutAngularVelocity = FVector::ZeroVector;
 
+	// The first time to analyze
 	UPROPERTY(EditAnywhere, Category = "Trajectory", meta = (Input))
 	float Time1 = 0.0f;
 
+	// The second time to analyze
 	UPROPERTY(EditAnywhere, Category = "Trajectory", meta = (Input))
 	float Time2 = 0.0f;
 
+	// If True the trajectory will be extrapolated beyond its bounds
 	UPROPERTY(EditAnywhere, Category = "Trajectory", meta = (Input))
 	bool bExtrapolate = false;
 };
\ No newline at end of file
diff --git a/Engine/Plugins/Experimental/UAF/UAF/Source/UAF/Internal/Graph/RigUnit_AnimNextBase.h b/Engine/Plugins/Experimental/UAF/UAF/Source/UAF/Internal/Graph/RigUnit_AnimNextBase.h
--- a/Engine/Plugins/Experimental/UAF/UAF/Source/UAF/Internal/Graph/RigUnit_AnimNextBase.h
+++ b/Engine/Plugins/Experimental/UAF/UAF/Source/UAF/Internal/Graph/RigUnit_AnimNextBase.h
@@ -7,7 +7,10 @@
 #include "AnimNextExecuteContext.h"
 #include "RigUnit_AnimNextBase.generated.h"
 
-USTRUCT(meta=(ExecuteContext="FAnimNextExecuteContext"))
+/*
+ * base class for all UAF rigvm nodes
+ */
+USTRUCT(meta=(ExecuteContext="FAnimNextExecuteContext", DocumentationPolicy = "Strict"))
 struct FRigUnit_AnimNextBase : public FRigVMStruct
 {
 	GENERATED_BODY()
diff --git a/Engine/Plugins/Experimental/UAF/UAF/Source/UAF/Internal/Module/RigUnit_AnimNextModuleEventDependencies.h b/Engine/Plugins/Experimental/UAF/UAF/Source/UAF/Internal/Module/RigUnit_AnimNextModuleEventDependencies.h
--- a/Engine/Plugins/Experimental/UAF/UAF/Source/UAF/Internal/Module/RigUnit_AnimNextModuleEventDependencies.h
+++ b/Engine/Plugins/Experimental/UAF/UAF/Source/UAF/Internal/Module/RigUnit_AnimNextModuleEventDependencies.h
@@ -19,6 +19,9 @@ public:
 	TObjectPtr<URigVMNode> Node = nullptr;
 };
 
+/*
+ * the base class for all dependencies nodes
+ */
 USTRUCT(meta=(Hidden, Category="Dependencies", NodeColor="0, 1, 1", Keywords="Prerequisite,Subsequent,TickFunction"))
 struct FRigUnit_AnimNextModuleDependenciesBase : public FRigUnit_AnimNextBase
 {
diff --git a/Engine/Plugins/Experimental/UAF/UAF/Source/UAF/Internal/Variables/RigUnit_ResolveUniversalObjectLocator.h b/Engine/Plugins/Experimental/UAF/UAF/Source/UAF/Internal/Variables/RigUnit_ResolveUniversalObjectLocator.h
--- a/Engine/Plugins/Experimental/UAF/UAF/Source/UAF/Internal/Variables/RigUnit_ResolveUniversalObjectLocator.h
+++ b/Engine/Plugins/Experimental/UAF/UAF/Source/UAF/Internal/Variables/RigUnit_ResolveUniversalObjectLocator.h
@@ -19,9 +19,11 @@ struct FRigUnit_ResolveUniversalObjectLocator : public FRigUnit_AnimNextBase
 	RIGVM_METHOD()
 	UE_API virtual void Execute() override;
 
+	// The object locator to resolve
 	UPROPERTY(meta = (Input))
 	FUniversalObjectLocator Locator;
 
+	// The resulting resolved object
 	UPROPERTY(meta = (Output))
 	TObjectPtr<UObject> Object;
 };
diff --git a/Engine/Plugins/Experimental/UAF/UAF/Source/UAF/Private/Graph/RigUnit_GetSocketTransform.h b/Engine/Plugins/Experimental/UAF/UAF/Source/UAF/Private/Graph/RigUnit_GetSocketTransform.h
--- a/Engine/Plugins/Experimental/UAF/UAF/Source/UAF/Private/Graph/RigUnit_GetSocketTransform.h
+++ b/Engine/Plugins/Experimental/UAF/UAF/Source/UAF/Private/Graph/RigUnit_GetSocketTransform.h
@@ -17,15 +17,19 @@ struct FRigUnit_GetSocketTransform : public FRigUnit_AnimNextBase
 	RIGVM_METHOD()
 	void Execute();
 
+	// The scene component to look up a socket transform for
 	UPROPERTY(EditAnywhere, Category = "Graph", meta = (Input))
 	TObjectPtr<USceneComponent> SceneComponent = nullptr;
 
+	// The name of the socket transform to retrieve
 	UPROPERTY(EditAnywhere, Category = "Graph", meta = (Input))
 	FName SocketName = NAME_None;
 
+	// The space to retrieve the transform in
 	UPROPERTY(EditAnywhere, Category = "Graph", meta = (Input))
 	TEnumAsByte<ERelativeTransformSpace> TransformSpace = ERelativeTransformSpace::RTS_World;
 
+	// The resulting socket transform
 	UPROPERTY(EditAnywhere, Category = "Graph", meta = (Output))
 	FTransform Result = FTransform::Identity;
 };
diff --git a/Engine/Plugins/Experimental/UAF/UAF/Source/UAF/Private/Graph/RigUnit_IsAnimBPDriven.h b/Engine/Plugins/Experimental/UAF/UAF/Source/UAF/Private/Graph/RigUnit_IsAnimBPDriven.h
--- a/Engine/Plugins/Experimental/UAF/UAF/Source/UAF/Private/Graph/RigUnit_IsAnimBPDriven.h
+++ b/Engine/Plugins/Experimental/UAF/UAF/Source/UAF/Private/Graph/RigUnit_IsAnimBPDriven.h
@@ -16,9 +16,11 @@ struct FRigUnit_IsAnimBPDriven : public FRigUnit_AnimNextBase
 	RIGVM_METHOD()
 	void Execute();
 
+	// The mesh component to check
 	UPROPERTY(EditAnywhere, Category = "Graph", meta = (Input))
 	TObjectPtr<USkeletalMeshComponent> MeshComponent;
 
+	// True if the skeletal mesh component is being driven by an Animation Blueprint
 	UPROPERTY(EditAnywhere, Category = "Graph", meta = (Output))
 	bool Result = false;
 };
diff --git a/Engine/Plugins/Experimental/UAF/UAFAnimGraph/Source/UAFAnimGraph/Internal/Graph/RigUnit_AnimNextGraphEvaluator.h b/Engine/Plugins/Experimental/UAF/UAFAnimGraph/Source/UAFAnimGraph/Internal/Graph/RigUnit_AnimNextGraphEvaluator.h
--- a/Engine/Plugins/Experimental/UAF/UAFAnimGraph/Source/UAFAnimGraph/Internal/Graph/RigUnit_AnimNextGraphEvaluator.h
+++ b/Engine/Plugins/Experimental/UAF/UAFAnimGraph/Source/UAFAnimGraph/Internal/Graph/RigUnit_AnimNextGraphEvaluator.h
@@ -21,6 +21,7 @@ struct FRigUnit_AnimNextGraphEvaluator : public FRigUnit_AnimNextBase
 {
 	GENERATED_BODY()
 
+	// The execute context used for the execution flow
 	UPROPERTY(meta = (Input, Output))
 	FAnimNextExecuteContext ExecuteContext;
 
diff --git a/Engine/Plugins/Experimental/UAF/UAFAnimGraph/Source/UAFAnimGraph/Private/Graph/RigUnit_AnimNextSetNotifyContext.h b/Engine/Plugins/Experimental/UAF/UAFAnimGraph/Source/UAFAnimGraph/Private/Graph/RigUnit_AnimNextSetNotifyContext.h
--- a/Engine/Plugins/Experimental/UAF/UAFAnimGraph/Source/UAFAnimGraph/Private/Graph/RigUnit_AnimNextSetNotifyContext.h
+++ b/Engine/Plugins/Experimental/UAF/UAFAnimGraph/Source/UAFAnimGraph/Private/Graph/RigUnit_AnimNextSetNotifyContext.h
@@ -21,6 +21,7 @@ struct FRigUnit_AnimNextSetNotifyContext : public FRigUnit_AnimNextBase
 	UPROPERTY(meta = (Input))
 	TObjectPtr<USkeletalMeshComponent> SkeletalMeshComponent;
 
+	// The execute context to use for the given skeletal mesh component
 	UPROPERTY(meta = (Input, Output))
 	FAnimNextExecuteContext ExecuteContext;
 };
diff --git a/Engine/Plugins/Experimental/UAF/UAFAnimGraph/Source/UAFAnimGraph/Private/Graph/RigUnit_HasValidEvaluationModifier.h b/Engine/Plugins/Experimental/UAF/UAFAnimGraph/Source/UAFAnimGraph/Private/Graph/RigUnit_HasValidEvaluationModifier.h
--- a/Engine/Plugins/Experimental/UAF/UAFAnimGraph/Source/UAFAnimGraph/Private/Graph/RigUnit_HasValidEvaluationModifier.h
+++ b/Engine/Plugins/Experimental/UAF/UAFAnimGraph/Source/UAFAnimGraph/Private/Graph/RigUnit_HasValidEvaluationModifier.h
@@ -8,6 +8,7 @@
 
 #include "RigUnit_HasValidEvaluationModifier.generated.h"
 
+// Used to determine whether an evaluation modifier is valid (indicates active injection)
 USTRUCT(meta=(DisplayName="Has Valid Evaluation Modifier", Category="Animation Next Injection", NodeColor="0, 1, 1"))
 struct FRigUnit_HasValidEvaluationModifier : public FRigUnit_AnimNextBase
 {
diff --git a/Engine/Plugins/Experimental/UAF/UAFPoseSearch/Source/UAFPoseSearch/Private/RigUnit_DebugDrawTrajectory.h b/Engine/Plugins/Experimental/UAF/UAFPoseSearch/Source/UAFPoseSearch/Private/RigUnit_DebugDrawTrajectory.h
--- a/Engine/Plugins/Experimental/UAF/UAFPoseSearch/Source/UAFPoseSearch/Private/RigUnit_DebugDrawTrajectory.h
+++ b/Engine/Plugins/Experimental/UAF/UAFPoseSearch/Source/UAFPoseSearch/Private/RigUnit_DebugDrawTrajectory.h
@@ -8,6 +8,9 @@
 
 #include "RigUnit_DebugDrawTrajectory.generated.h"
 
+/*
+ * A debug drawing node to draw a transform trajectory
+ */
 USTRUCT(meta = (DisplayName = "Debug Draw Trajectory", Category = "Trajectory", NodeColor = "0, 1, 1", Keywords = "Mover, Trajectory, Debug"))
 struct FRigUnit_DebugDrawTrajectory: public FRigUnit_AnimNextBase
 {
@@ -16,15 +19,19 @@ struct FRigUnit_DebugDrawTrajectory: public FRigUnit_AnimNextBase
 	RIGVM_METHOD()
 	void Execute();
 
+	// The trajectory to draw
 	UPROPERTY(EditAnywhere, Category = "Debug Draw - Trajectory", meta = (Input, Output))
 	FTransformTrajectory Trajectory;
 
+	// The line thickness to use for the debug drawing
 	UPROPERTY(EditAnywhere, Category = "Debug Draw - Trajectory", meta = (Input))
 	float DebugThickness = 0.f;
 
+	// The transform offset to use for the debug drawing
 	UPROPERTY(EditAnywhere, Category = "Debug Draw - Trajectory", meta = (Input))
 	float DebugOffset = 0.f;
 
+	// If False the debug drawing will be skipped
 	UPROPERTY(EditAnywhere, Category = "Debug Draw - Trajectory", meta = (Input))
 	bool Enabled = false;
 
diff --git a/Engine/Plugins/Experimental/UAF/UAFPoseSearch/Source/UAFPoseSearch/Private/RigUnit_GenerateCharacterMovementComponentTrajectory.h b/Engine/Plugins/Experimental/UAF/UAFPoseSearch/Source/UAFPoseSearch/Private/RigUnit_GenerateCharacterMovementComponentTrajectory.h
--- a/Engine/Plugins/Experimental/UAF/UAFPoseSearch/Source/UAFPoseSearch/Private/RigUnit_GenerateCharacterMovementComponentTrajectory.h
+++ b/Engine/Plugins/Experimental/UAF/UAFPoseSearch/Source/UAFPoseSearch/Private/RigUnit_GenerateCharacterMovementComponentTrajectory.h
@@ -10,6 +10,9 @@
 
 #include "RigUnit_GenerateCharacterMovementComponentTrajectory.generated.h"
 
+/*
+ * A utility node to generate a character movement trajectory from a character movement component
+ */
 USTRUCT(meta = (DisplayName = "Generate Trajectory from Character Movement Component", Category = "Character Movement", NodeColor = "0, 1, 1", Keywords = "Trajectory"))
 struct FRigUnit_GenerateCharacterMovementComponentTrajectory : public FRigUnit_AnimNextBase
 {
@@ -18,31 +21,39 @@ struct FRigUnit_GenerateCharacterMovementComponentTrajectory : public FRigUnit_A
 	RIGVM_METHOD()
 	void Execute();
 
+	// The input character movement component to create the trajectory for
 	UPROPERTY(EditAnywhere, Category = "Trajectory", meta = (Input))
 	TObjectPtr<UCharacterMovementComponent> CharacterMovementComponent;
 
-	UPROPERTY(EditAnywhere, Category = "Trajectory", meta = (Input))
 	// This should be the most recent simulation time that was used to get us to our current state
+	UPROPERTY(EditAnywhere, Category = "Trajectory", meta = (Input))
 	float DeltaTime = 0.f;
 
+	// The sampling interval to use for the history
 	UPROPERTY(EditAnywhere, Category = "Trajectory", meta = (Input))
 	float HistorySamplingInterval = -1.f;
 
+	// The upper limit of samples to use for the history
 	UPROPERTY(EditAnywhere, Category = "Trajectory", meta = (Input))
 	int32 NumHistorySamples = 30;
 
+	// The interval to use for prediction
 	UPROPERTY(EditAnywhere, Category = "Trajectory", meta = (Input))
 	float PredictionSamplingInterval = 0.1f;
 
+	// The upper limit of samples to use for prediction
 	UPROPERTY(EditAnywhere, Category = "Trajectory", meta = (Input))
 	int32 NumPredictionSamples = 15;
 
+	// The input trajectory data to use
 	UPROPERTY(EditAnywhere, Category = "Trajectory", meta = (Input))
 	FPoseSearchTrajectoryData TrajectoryData;
-	
+
+	// The resulting trajectory
 	UPROPERTY(EditAnywhere, Category = "Trajectory", meta = (Input, Output))
 	FTransformTrajectory InOutTrajectory;
 
+	// The resulting yaw at the last update
 	UPROPERTY(EditAnywhere, Category = "Trajectory", meta = (Input, Output))
 	float InOutDesiredControllerYawLastUpdate = 0.0f;
 	
diff --git a/Engine/Plugins/Experimental/UAF/UAFPoseSearch/Source/UAFPoseSearch/Private/RigUnit_PoseSearchDatabaseGetTags.h b/Engine/Plugins/Experimental/UAF/UAFPoseSearch/Source/UAFPoseSearch/Private/RigUnit_PoseSearchDatabaseGetTags.h
--- a/Engine/Plugins/Experimental/UAF/UAFPoseSearch/Source/UAFPoseSearch/Private/RigUnit_PoseSearchDatabaseGetTags.h
+++ b/Engine/Plugins/Experimental/UAF/UAFPoseSearch/Source/UAFPoseSearch/Private/RigUnit_PoseSearchDatabaseGetTags.h
@@ -16,9 +16,11 @@ struct FRigUnit_PoseSearchDatabaseGetTags : public FRigUnit_AnimNextBase
 	RIGVM_METHOD()
 	virtual void Execute() override;
 
+	// The data base to retrieve the meta data tags from
 	UPROPERTY(EditAnywhere, Transient, Category = "Pose Search", meta = (Input))
 	TObjectPtr<UPoseSearchDatabase> Database;
 
+	// The resulting array of meta data tags
 	UPROPERTY(EditAnywhere, Transient, Category = "Pose Search", meta = (Output))
 	TArray<FName> OutTags;
 };
diff --git a/Engine/Plugins/Experimental/UAF/UAFTests/Source/UAFAnimGraphTestSuite/Private/GCTestsUtil.h b/Engine/Plugins/Experimental/UAF/UAFTests/Source/UAFAnimGraphTestSuite/Private/GCTestsUtil.h
--- a/Engine/Plugins/Experimental/UAF/UAFTests/Source/UAFAnimGraphTestSuite/Private/GCTestsUtil.h
+++ b/Engine/Plugins/Experimental/UAF/UAFTests/Source/UAFAnimGraphTestSuite/Private/GCTestsUtil.h
@@ -25,7 +25,7 @@ struct FUAFTestAnimSequenceSharedData : public FAnimNextTraitSharedData
 	#undef TRAIT_LATENT_PROPERTIES_ENUMERATOR
 };
 
-USTRUCT()
+USTRUCT(meta = (DocumentationPolicy = "None"))
 struct FRigUnit_AnimSequenceTest : public FRigUnit_AnimNextBase
 {
 	GENERATED_BODY()
