Control Rig: Fix TangentAtParam returns bad value when param is 1
#rnx

#Synced-CL 48500891

diff --git a/Engine/Plugins/Animation/ControlRigSpline/Source/ControlRigSpline/Private/ControlRigSplineTypes.cpp b/Engine/Plugins/Animation/ControlRigSpline/Source/ControlRigSpline/Private/ControlRigSplineTypes.cpp
--- a/Engine/Plugins/Animation/ControlRigSpline/Source/ControlRigSpline/Private/ControlRigSplineTypes.cpp
+++ b/Engine/Plugins/Animation/ControlRigSpline/Source/ControlRigSpline/Private/ControlRigSplineTypes.cpp
@@ -623,16 +623,20 @@ FVector FControlRigSpline::TangentAtParam(const float InParam) const
 	{
 		return FVector();
 	}
-	
-	
 
 	const float ClampedU = FMath::Clamp<float>(InParam, 0.f, 1.f);
 
 	const int32 LastIndex = SplineData->SamplesArray.Num() - 1;
 	const float fIndexPrev = ClampedU * LastIndex;	
-	const int32 IndexPrev = FMath::Floor<int32>(fIndexPrev);
+	int32 IndexPrev = FMath::Floor<int32>(fIndexPrev);
 	const float ULocal = fIndexPrev - IndexPrev;
-	const int32 IndexNext = (IndexPrev < LastIndex) ? IndexPrev + 1 : IndexPrev;
+	int32 IndexNext = IndexPrev + 1;
+	
+	if (IndexNext > LastIndex)
+	{
+		IndexPrev--;
+		IndexNext--;
+	}
 	
 	const FVector TangentStart = SplineData->SamplesArray[IndexNext].GetTranslation() - SplineData->SamplesArray[IndexPrev].GetTranslation();
 	if (IndexNext + 1 < SplineData->SamplesArray.Num())
